add missing virtual destructors

diff --git a/Engine/Source/Runtime/InteractiveToolsFramework/Public/ToolContextInterfaces.h b/Engine/Source/Runtime/InteractiveToolsFramework/Public/ToolContextInterfaces.h
--- a/Engine/Source/Runtime/InteractiveToolsFramework/Public/ToolContextInterfaces.h
+++ b/Engine/Source/Runtime/InteractiveToolsFramework/Public/ToolContextInterfaces.h
@@ -88,6 +88,7 @@ enum class EToolMessageLevel
 class IToolsContextTransactionsAPI
 {
 public:
+	virtual ~IToolsContextTransactionsAPI() {}
 
 	/**
 	 * Request that context display message information.
@@ -137,6 +138,8 @@ public:
 class IToolsContextRenderAPI
 {
 public:
+	virtual ~IToolsContextRenderAPI() {}
+
 	/** @return Current PDI */
 	virtual FPrimitiveDrawInterface* GetPrimitiveDrawInterface() = 0;
 };
@@ -153,6 +156,8 @@ public:
 class IToolsContextAssetAPI
 {
 public:
+	virtual ~IToolsContextAssetAPI() {}
+
 	/** Get default path to save assets in. For example the currently-visible path in the Editor. */
 	virtual FString GetActiveAssetFolderPath() = 0;
 
