Add ToolManager::PostMessage, FInputDeviceRay
#rnx

diff --git a/Engine/Source/Runtime/InteractiveToolsFramework/Private/InteractiveToolManager.cpp b/Engine/Source/Runtime/InteractiveToolsFramework/Private/InteractiveToolManager.cpp
--- a/Engine/Source/Runtime/InteractiveToolsFramework/Private/InteractiveToolManager.cpp
+++ b/Engine/Source/Runtime/InteractiveToolsFramework/Private/InteractiveToolManager.cpp
@@ -215,7 +215,15 @@ void UInteractiveToolManager::Render(IToolsContextRenderAPI* RenderAPI)
 }
 
 
+void UInteractiveToolManager::PostMessage(const TCHAR* Message, EToolMessageLevel Level)
+{
+	TransactionsAPI->PostMessage(Message, Level);
+}
 
+void UInteractiveToolManager::PostMessage(const FString& Message, EToolMessageLevel Level)
+{
+	TransactionsAPI->PostMessage(*Message, Level);
+}
 
 void UInteractiveToolManager::PostInvalidation()
 {
diff --git a/Engine/Source/Runtime/InteractiveToolsFramework/Public/InputState.h b/Engine/Source/Runtime/InteractiveToolsFramework/Public/InputState.h
--- a/Engine/Source/Runtime/InteractiveToolsFramework/Public/InputState.h
+++ b/Engine/Source/Runtime/InteractiveToolsFramework/Public/InputState.h
@@ -28,6 +28,38 @@ ENUM_CLASS_FLAGS(EInputDevices);
 
 
 
+/**
+ * FInputDeviceRay represents a 3D ray created based on an input device.
+ * If the device is a 2D input device like a mouse, then the ray may
+ * have an associated 2D screen position.
+ */
+struct FInputDeviceRay
+{
+	/** 3D ray in 3D scene, in world coordinates */
+	FRay WorldRay;
+
+	/** If true, WorldRay has 2D device position coordinates */
+	bool bHas2D;
+
+	/** 2D device position coordinates associated with the ray */
+	FVector2D ScreenPosition;
+
+
+	FInputDeviceRay(const FRay & WorldRayIn)
+	{
+		WorldRay = WorldRayIn;
+		bHas2D = false;
+		ScreenPosition = FVector2D(0, 0);
+	}
+
+	FInputDeviceRay(const FRay & WorldRayIn, const FVector2D& ScreenPositionIn)
+	{
+		WorldRay = WorldRayIn;
+		bHas2D = true;
+		ScreenPosition = ScreenPositionIn;
+	}
+};
+
 
 /**
  * Current State of a physical device button (mouse, key, etc) at a point in time.
diff --git a/Engine/Source/Runtime/InteractiveToolsFramework/Public/InteractiveToolManager.h b/Engine/Source/Runtime/InteractiveToolsFramework/Public/InteractiveToolManager.h
--- a/Engine/Source/Runtime/InteractiveToolsFramework/Public/InteractiveToolManager.h
+++ b/Engine/Source/Runtime/InteractiveToolsFramework/Public/InteractiveToolManager.h
@@ -132,7 +132,12 @@ public:
 	// Functions that Tools can call to interact with Transactions API
 	//
 	
-	
+	/** Post a message via the Transactions API */
+	virtual void PostMessage(const TCHAR* Message, EToolMessageLevel Level);
+
+	/** Post a message via the Transactions API */
+	virtual void PostMessage(const FString& Message, EToolMessageLevel Level);
+
 	/** Request an Invalidation via the Transactions API (ie to cause a repaint, etc) */
 	virtual void PostInvalidation();
 
