Base InputBehavior and InteractiveTool implementations for single-click-in-scene and click-drag-on-mesh interactions
#rnx

diff --git a/Engine/Source/Runtime/InteractiveToolsFramework/Private/BaseBehaviors/AnyButtonInputBehavior.cpp b/Engine/Source/Runtime/InteractiveToolsFramework/Private/BaseBehaviors/AnyButtonInputBehavior.cpp
new file mode 100644
--- /dev/null
+++ b/Engine/Source/Runtime/InteractiveToolsFramework/Private/BaseBehaviors/AnyButtonInputBehavior.cpp
@@ -0,0 +1,112 @@
+// Copyright 1998-2019 Epic Games, Inc. All Rights Reserved.
+
+#include "BaseBehaviors/AnyButtonInputBehavior.h"
+
+
+
+UAnyButtonInputBehavior::UAnyButtonInputBehavior()
+{
+	ButtonNumber = 0;
+}
+
+
+EInputDevices UAnyButtonInputBehavior::GetSupportedDevices()
+{
+	return EInputDevices::Mouse;
+}
+
+
+bool UAnyButtonInputBehavior::IsPressed(const FInputDeviceState& input)
+{
+	if (input.IsFromDevice(EInputDevices::Mouse)) 
+	{
+		ActiveDevice = EInputDevices::Mouse;
+		return GetMouseButtonState(input).bPressed;
+	} 
+	else if (input.IsFromDevice(EInputDevices::TabletFingers))
+	{
+		ActiveDevice = EInputDevices::TabletFingers;
+		//return input.TouchPressed;  // not implemented yet
+		return false;
+	}
+	return false;
+}
+
+bool UAnyButtonInputBehavior::IsDown(const FInputDeviceState& input)
+{
+	if (input.IsFromDevice(EInputDevices::Mouse))
+	{
+		ActiveDevice = EInputDevices::Mouse;
+		return GetMouseButtonState(input).bDown;
+	}
+	return false;
+}
+
+bool UAnyButtonInputBehavior::IsReleased(const FInputDeviceState& input)
+{
+	if (input.IsFromDevice(EInputDevices::Mouse))
+	{
+		ActiveDevice = EInputDevices::Mouse;
+		return GetMouseButtonState(input).bReleased;
+	}
+	return false;
+}
+
+FVector2D UAnyButtonInputBehavior::GetClickPoint(const FInputDeviceState& input)
+{
+	if (input.IsFromDevice(EInputDevices::Mouse))
+	{
+		ActiveDevice = EInputDevices::Mouse;
+		return input.Mouse.Position2D;
+	}
+	return FVector2D::ZeroVector;
+}
+
+
+FRay UAnyButtonInputBehavior::GetWorldRay(const FInputDeviceState& input)
+{
+	if (input.IsFromDevice(EInputDevices::Mouse))
+	{
+		ActiveDevice = EInputDevices::Mouse;
+		return input.Mouse.WorldRay;
+	}
+	return FRay(FVector::ZeroVector, FVector(0, 0, 1), true);
+
+}
+
+
+FInputDeviceRay UAnyButtonInputBehavior::GetDeviceRay(const FInputDeviceState& input)
+{
+	if (input.IsFromDevice(EInputDevices::Mouse))
+	{
+		ActiveDevice = EInputDevices::Mouse;
+		return FInputDeviceRay(input.Mouse.WorldRay, input.Mouse.Position2D);
+	}
+	return FInputDeviceRay(FRay(FVector::ZeroVector, FVector(0, 0, 1), true));
+
+}
+
+
+EInputDevices UAnyButtonInputBehavior::GetActiveDevice() const
+{
+	return ActiveDevice;
+}
+
+
+
+
+FDeviceButtonState UAnyButtonInputBehavior::GetMouseButtonState(const FInputDeviceState& input)
+{
+	if (ButtonNumber == 2)
+	{
+		return input.Mouse.Right;
+	}
+	else if (ButtonNumber == 1)
+	{
+		return input.Mouse.Middle;
+	}
+	else
+	{
+		return input.Mouse.Left;
+	}
+}
\ No newline at end of file
diff --git a/Engine/Source/Runtime/InteractiveToolsFramework/Private/BaseBehaviors/MouseHoverBehavior.cpp b/Engine/Source/Runtime/InteractiveToolsFramework/Private/BaseBehaviors/MouseHoverBehavior.cpp
new file mode 100644
--- /dev/null
+++ b/Engine/Source/Runtime/InteractiveToolsFramework/Private/BaseBehaviors/MouseHoverBehavior.cpp
@@ -0,0 +1,37 @@
+// Copyright 1998-2019 Epic Games, Inc. All Rights Reserved.
+
+#include "BaseBehaviors/MouseHoverBehavior.h"
+
+
+UMouseHoverBehavior::UMouseHoverBehavior()
+{
+	Target = nullptr;
+}
+
+EInputDevices UMouseHoverBehavior::GetSupportedDevices()
+{
+	return EInputDevices::Mouse;
+}
+
+void UMouseHoverBehavior::Initialize(IHoverBehaviorTarget* TargetIn)
+{
+	this->Target = TargetIn;
+}
+
+bool UMouseHoverBehavior::WantsHoverEvents()
+{
+	return true;
+}
+
+void UMouseHoverBehavior::UpdateHover(const FInputDeviceState& input)
+{
+	if (Target != nullptr)
+	{
+		Target->OnUpdateHover( FInputDeviceRay(input.Mouse.WorldRay, input.Mouse.Position2D) );
+	}
+}
+
+void UMouseHoverBehavior::EndHover(const FInputDeviceState& input)
+{
+}
+
diff --git a/Engine/Source/Runtime/InteractiveToolsFramework/Private/BaseBehaviors/SingleClickBehavior.cpp b/Engine/Source/Runtime/InteractiveToolsFramework/Private/BaseBehaviors/SingleClickBehavior.cpp
new file mode 100644
--- /dev/null
+++ b/Engine/Source/Runtime/InteractiveToolsFramework/Private/BaseBehaviors/SingleClickBehavior.cpp
@@ -0,0 +1,65 @@
+// Copyright 1998-2019 Epic Games, Inc. All Rights Reserved.
+
+#include "BaseBehaviors/SingleClickBehavior.h"
+
+
+
+USingleClickToolBehavior::USingleClickToolBehavior()
+{
+	HitTestOnRelease = true;
+}
+
+
+void USingleClickToolBehavior::Initialize(IClickBehaviorTarget* TargetIn)
+{
+	this->Target = TargetIn;
+}
+
+
+FInputCaptureRequest USingleClickToolBehavior::WantsCapture(const FInputDeviceState& input)
+{
+	if (IsPressed(input))
+	{
+		if ( Target->IsHitByClick(GetDeviceRay(input)) )
+		{
+			return FInputCaptureRequest::Begin(this, EInputCaptureSide::Any);
+		}
+	}
+	return FInputCaptureRequest::Ignore();
+}
+
+
+FInputCaptureUpdate USingleClickToolBehavior::BeginCapture(const FInputDeviceState& input, EInputCaptureSide eSide)
+{
+	return FInputCaptureUpdate::Begin(this, EInputCaptureSide::Any);
+}
+
+
+FInputCaptureUpdate USingleClickToolBehavior::UpdateCapture(const FInputDeviceState& input, const FInputCaptureData& data)
+{
+	if (IsReleased(input)) 
+	{
+		if (HitTestOnRelease == false || 
+			Target->IsHitByClick(GetDeviceRay(input)) )
+		{
+			Clicked(input, data);
+		}
+
+		return FInputCaptureUpdate::End();
+	}
+
+	return FInputCaptureUpdate::Continue();
+}
+
+
+void USingleClickToolBehavior::ForceEndCapture(const FInputCaptureData& data)
+{
+	// nothing to do
+}
+
+
+void USingleClickToolBehavior::Clicked(const FInputDeviceState& input, const FInputCaptureData& data)
+{
+	Target->OnClicked(GetDeviceRay(input));
+}
+
diff --git a/Engine/Source/Runtime/InteractiveToolsFramework/Private/BaseTools/MeshSurfacePointTool.cpp b/Engine/Source/Runtime/InteractiveToolsFramework/Private/BaseTools/MeshSurfacePointTool.cpp
new file mode 100644
--- /dev/null
+++ b/Engine/Source/Runtime/InteractiveToolsFramework/Private/BaseTools/MeshSurfacePointTool.cpp
@@ -0,0 +1,166 @@
+// Copyright 1998-2019 Epic Games, Inc. All Rights Reserved.
+
+#include "BaseTools/MeshSurfacePointTool.h"
+#include "BaseBehaviors/MouseHoverBehavior.h"
+#include "InteractiveToolManager.h"
+#include "ToolBuilderUtil.h"
+
+
+/*
+ * ToolBuilder
+ */
+
+bool UMeshSurfacePointToolBuilder::CanBuildTool(const FToolBuilderState& SceneState) const
+{
+	return ToolBuilderUtil::CountComponents(SceneState, ToolBuilderUtil::IsMeshDescriptionSourceComponent) == 1;
+}
+
+UInteractiveTool* UMeshSurfacePointToolBuilder::BuildTool(const FToolBuilderState& SceneState) const
+{
+	UMeshSurfacePointTool* NewTool = CreateNewTool(SceneState);
+	InitializeNewTool(NewTool, SceneState);
+	return NewTool;
+}
+
+UMeshSurfacePointTool* UMeshSurfacePointToolBuilder::CreateNewTool(const FToolBuilderState& SceneState) const
+{
+	return NewObject<UMeshSurfacePointTool>(SceneState.ToolManager);
+}
+
+void UMeshSurfacePointToolBuilder::InitializeNewTool(UMeshSurfacePointTool* NewTool, const FToolBuilderState& SceneState) const
+{
+	UObject* MeshComponent = ToolBuilderUtil::FindFirstComponent(SceneState, ToolBuilderUtil::IsMeshDescriptionSourceComponent);
+	check(MeshComponent != nullptr);
+	NewTool->SetSelection(
+		SceneState.SourceBuilder->MakeMeshDescriptionSource(MeshComponent) );
+}
+
+
+/*
+ * Tool
+ */
+
+void UMeshSurfacePointTool::SetSelection(TUniquePtr<IMeshDescriptionSource> MeshSourceIn)
+{
+	this->MeshSource = TUniquePtr<IMeshDescriptionSource>(MoveTemp(MeshSourceIn));
+}
+
+
+void UMeshSurfacePointTool::Setup()
+{
+	UInteractiveTool::Setup();
+
+	bShiftToggle = false;
+
+	// add input behaviors
+	UMeshSurfacePointToolMouseBehavior* mouseBehavior = NewObject<UMeshSurfacePointToolMouseBehavior>();
+	mouseBehavior->Initialize(this);
+	AddInputBehavior(mouseBehavior);
+
+	UMouseHoverBehavior* hoverBehavior = NewObject<UMouseHoverBehavior>();
+	hoverBehavior->Initialize(this);
+	AddInputBehavior(hoverBehavior);
+}
+
+
+bool UMeshSurfacePointTool::HitTest(const FRay& Ray, FHitResult& OutHit)
+{
+	return MeshSource->HitTest(Ray, OutHit);
+}
+
+
+void UMeshSurfacePointTool::OnBeginDrag(const FRay& Ray)
+{
+
+}
+
+
+void UMeshSurfacePointTool::OnUpdateDrag(const FRay& Ray)
+{
+	FHitResult OutHit;
+	if ( HitTest(Ray, OutHit) ) 
+	{
+		GetToolManager()->PostMessage( 
+			FString::Printf(TEXT("[UMeshSurfacePointTool::OnUpdateDrag] Hit triangle index %d at ray distance %f"), OutHit.FaceIndex, OutHit.Distance), EToolMessageLevel::Internal);
+	}
+}
+
+void UMeshSurfacePointTool::OnEndDrag(const FRay& Ray)
+{
+	//GetToolManager()->PostMessage(TEXT("UMeshSurfacePointTool::OnEndDrag!"), EToolMessageLevel::Internal);
+}
+
+
+void UMeshSurfacePointTool::SetShiftToggle(bool bShiftDown)
+{
+	bShiftToggle = bShiftDown;
+}
+
+
+
+
+/*
+ * Mouse Input Behavior
+ */
+
+
+void UMeshSurfacePointToolMouseBehavior::Initialize(UMeshSurfacePointTool* ToolIn)
+{
+	this->Tool = ToolIn;
+	bInDragCapture = false;
+}
+
+
+FInputCaptureRequest UMeshSurfacePointToolMouseBehavior::WantsCapture(const FInputDeviceState& input)
+{
+	if ( IsPressed(input) )
+	{
+		FHitResult OutHit;
+		if (Tool->HitTest(input.Mouse.WorldRay, OutHit))
+		{
+			return FInputCaptureRequest::Begin(this, EInputCaptureSide::Any, OutHit.Distance);
+		}
+	}
+	return FInputCaptureRequest::Ignore();
+}
+
+FInputCaptureUpdate UMeshSurfacePointToolMouseBehavior::BeginCapture(const FInputDeviceState& input, EInputCaptureSide eSide)
+{
+	Tool->SetShiftToggle(input.bShiftKeyDown);
+	Tool->OnBeginDrag(input.Mouse.WorldRay);
+	LastWorldRay = input.Mouse.WorldRay;
+	bInDragCapture = true;
+	return FInputCaptureUpdate::Begin(this, EInputCaptureSide::Any);
+}
+
+FInputCaptureUpdate UMeshSurfacePointToolMouseBehavior::UpdateCapture(const FInputDeviceState& input, const FInputCaptureData& data)
+{
+	LastWorldRay = input.Mouse.WorldRay;
+
+	if ( IsReleased(input) ) 
+	{
+		Tool->OnEndDrag(input.Mouse.WorldRay);
+		bInDragCapture = false;
+		return FInputCaptureUpdate::End();
+	}
+
+	Tool->OnUpdateDrag(input.Mouse.WorldRay);
+	return FInputCaptureUpdate::Continue();
+}
+
+void UMeshSurfacePointToolMouseBehavior::ForceEndCapture(const FInputCaptureData& data)
+{
+	if (bInDragCapture)
+	{
+		Tool->OnEndDrag(LastWorldRay);
+		bInDragCapture = false;
+	}
+
+	// nothing to do
+}
+
+
+
+
+
+
diff --git a/Engine/Source/Runtime/InteractiveToolsFramework/Private/BaseTools/SingleClickTool.cpp b/Engine/Source/Runtime/InteractiveToolsFramework/Private/BaseTools/SingleClickTool.cpp
new file mode 100644
--- /dev/null
+++ b/Engine/Source/Runtime/InteractiveToolsFramework/Private/BaseTools/SingleClickTool.cpp
@@ -0,0 +1,56 @@
+// Copyright 1998-2019 Epic Games, Inc. All Rights Reserved.
+
+#include "BaseTools/SingleClickTool.h"
+#include "InteractiveToolManager.h"
+
+
+/*
+ * ToolBuilder
+ */
+
+bool USingleClickToolBuilder::CanBuildTool(const FToolBuilderState& SceneState) const
+{
+	return true;
+}
+
+UInteractiveTool* USingleClickToolBuilder::BuildTool(const FToolBuilderState& SceneState) const
+{
+	USingleClickTool* NewTool = NewObject<USingleClickTool>();
+	return NewTool;
+}
+
+
+
+/*
+ * Tool
+ */
+
+
+void USingleClickTool::Setup()
+{
+	UInteractiveTool::Setup();
+
+	// add default button input behaviors for devices
+	USingleClickToolBehavior* MouseBehavior = NewObject<USingleClickToolBehavior>();
+	MouseBehavior->Initialize(this);
+	AddInputBehavior(MouseBehavior);
+}
+
+
+bool USingleClickTool::IsHitByClick(const FInputDeviceRay& ClickPos)
+{
+	return true;
+}
+
+
+void USingleClickTool::OnClicked(const FInputDeviceRay& ClickPos)
+{
+	// print debug message
+	GetToolManager()->PostMessage( 
+		FString::Printf( TEXT("USingleClickTool::OnClicked at (%f,%f)"), ClickPos.ScreenPosition.X, ClickPos.ScreenPosition.Y), 
+		EToolMessageLevel::Internal );
+}
+
+
+
+
diff --git a/Engine/Source/Runtime/InteractiveToolsFramework/Public/BaseBehaviors/AnyButtonInputBehavior.h b/Engine/Source/Runtime/InteractiveToolsFramework/Public/BaseBehaviors/AnyButtonInputBehavior.h
new file mode 100644
--- /dev/null
+++ b/Engine/Source/Runtime/InteractiveToolsFramework/Public/BaseBehaviors/AnyButtonInputBehavior.h
@@ -0,0 +1,70 @@
+// Copyright 1998-2019 Epic Games, Inc. All Rights Reserved.
+
+#pragma once
+
+#include "CoreMinimal.h"
+#include "InputBehavior.h"
+#include "AnyButtonInputBehavior.generated.h"
+
+/**
+ * UAnyButtonInputBehavior is a base behavior that provides a generic
+ * interface to a TargetButton on a physical Input Device. You can subclass
+ * UAnyButtonInputBehavior to write InputBehaviors that can work independent
+ * of a particular device type or button, by using the UAnyButtonInputBehavior functions below.
+ * 
+ * The target device button is selected using the .ButtonNumber property, or you can
+ * override the relevant GetXButtonState() function if you need more control.
+ * 
+ *  @todo spatial controllers
+ *  @todo support tablet fingers
+ *  @todo support gamepad?
+ */
+UCLASS()
+class INTERACTIVETOOLSFRAMEWORK_API UAnyButtonInputBehavior : public UInputBehavior
+{
+	GENERATED_BODY()
+
+public:
+	UAnyButtonInputBehavior();
+
+	/** Return set of devices supported by this behavior */
+	virtual EInputDevices GetSupportedDevices() override;
+
+	/** @return true if Target Button has been pressed this frame */
+	virtual bool IsPressed(const FInputDeviceState& input);
+
+	/** @return true if Target Button is currently held down */
+	virtual bool IsDown(const FInputDeviceState& input);
+
+	/** @return true if Target Button was released this frame */
+	virtual bool IsReleased(const FInputDeviceState& input);
+	
+	/** @return current 2D position of Target Device, or zero if device does not have 2D position */
+	virtual FVector2D GetClickPoint(const FInputDeviceState& input);
+
+	/** @return current 3D world ray for Target Device position */
+	virtual FRay GetWorldRay(const FInputDeviceState& input);
+	
+	/** @return current 3D world ray and optional 2D position for Target Device */
+	virtual FInputDeviceRay GetDeviceRay(const FInputDeviceState& input);
+
+	/** @return last-active supported Device */
+	EInputDevices GetActiveDevice() const;
+
+protected:
+
+	/** 
+	 * Button number on target device. Button 0 is "default" button on all devices.
+	 * Mouse: Left=0, Middle=1, Right=2
+	 */
+	UPROPERTY()
+	int ButtonNumber;
+
+
+protected:
+	/** Which device is currently active */
+	EInputDevices ActiveDevice;
+
+	/** @return mouse button state for Target Button */
+	virtual FDeviceButtonState GetMouseButtonState(const FInputDeviceState& input);
+};
\ No newline at end of file
diff --git a/Engine/Source/Runtime/InteractiveToolsFramework/Public/BaseBehaviors/BehaviorTargetInterfaces.h b/Engine/Source/Runtime/InteractiveToolsFramework/Public/BaseBehaviors/BehaviorTargetInterfaces.h
new file mode 100644
--- /dev/null
+++ b/Engine/Source/Runtime/InteractiveToolsFramework/Public/BaseBehaviors/BehaviorTargetInterfaces.h
@@ -0,0 +1,47 @@
+// Copyright 1998-2019 Epic Games, Inc. All Rights Reserved.
+
+#pragma once
+
+#include "CoreMinimal.h"
+
+/**
+ * Functions required to apply standard "Click" state machines to a target object.
+ * See USingleClickBehavior for an example of this kind of state machine.
+ */
+class IClickBehaviorTarget
+{
+public:
+	virtual ~IClickBehaviorTarget() {}
+
+	/**
+	 * Test if target is hit by a click
+	 * @param ClickPos device position/ray at click point
+	 * @return true if target was hit by click ray/point
+	 */
+	virtual bool IsHitByClick(const FInputDeviceRay& ClickPos) = 0;
+
+
+	/**
+	 * Notify Target that click ocurred
+	 * @param ClickPos device position/ray at click point
+	 */
+	virtual void OnClicked(const FInputDeviceRay& ClickPos) = 0;
+};
+
+
+
+/**
+ * IHoverBehaviorTarget allows Behaviors to notify Tools/etc about
+ * device event data in a generic way, without requiring that all Tools
+ * know about the concept of Hovering.
+ */
+class IHoverBehaviorTarget
+{
+public:
+	virtual ~IHoverBehaviorTarget() {}
+
+	/**
+	 * Notify Target about a hover event
+	 */
+	virtual void OnUpdateHover(const FInputDeviceRay& DevicePos) = 0;
+};
diff --git a/Engine/Source/Runtime/InteractiveToolsFramework/Public/BaseBehaviors/MouseHoverBehavior.h b/Engine/Source/Runtime/InteractiveToolsFramework/Public/BaseBehaviors/MouseHoverBehavior.h
new file mode 100644
--- /dev/null
+++ b/Engine/Source/Runtime/InteractiveToolsFramework/Public/BaseBehaviors/MouseHoverBehavior.h
@@ -0,0 +1,39 @@
+// Copyright 1998-2019 Epic Games, Inc. All Rights Reserved.
+
+#pragma once
+
+#include "CoreMinimal.h"
+#include "InputBehavior.h"
+#include "InteractiveTool.h"
+#include "BaseBehaviors/BehaviorTargetInterfaces.h"
+#include "MouseHoverBehavior.generated.h"
+
+
+
+/**
+ * Trivial InputBehavior that forwards InputBehavior hover events to a Target object via
+ * the IHoverBehaviorTarget interface.
+ */
+UCLASS()
+class INTERACTIVETOOLSFRAMEWORK_API UMouseHoverBehavior : public UInputBehavior
+{
+	GENERATED_BODY()
+
+public:
+	UMouseHoverBehavior();
+
+	virtual void Initialize(IHoverBehaviorTarget* Target);
+
+	// UInputBehavior hover implementation
+
+	virtual EInputDevices GetSupportedDevices() override;
+
+	virtual bool WantsHoverEvents() override;
+	virtual void UpdateHover(const FInputDeviceState& input) override;
+	virtual void EndHover(const FInputDeviceState& input) override;
+
+protected:
+	IHoverBehaviorTarget* Target;
+};
+
+
diff --git a/Engine/Source/Runtime/InteractiveToolsFramework/Public/BaseBehaviors/SingleClickBehavior.h b/Engine/Source/Runtime/InteractiveToolsFramework/Public/BaseBehaviors/SingleClickBehavior.h
new file mode 100644
--- /dev/null
+++ b/Engine/Source/Runtime/InteractiveToolsFramework/Public/BaseBehaviors/SingleClickBehavior.h
@@ -0,0 +1,61 @@
+// Copyright 1998-2019 Epic Games, Inc. All Rights Reserved.
+
+#pragma once
+
+#include "CoreMinimal.h"
+#include "BaseBehaviors/AnyButtonInputBehavior.h"
+#include "BehaviorTargetInterfaces.h"
+#include "SingleClickBehavior.generated.h"
+
+
+/**
+ * USingleClickToolBehavior implements a standard "button-click"-style input behavior.
+ * The state machine works as follows:
+ *    1) on input-device-button-press, hit-test the target. If hit, begin capture
+ *    2) on input-device-button-release, hit-test the target. If hit, call Target::OnClicked(). If not hit, ignore click.
+ *    
+ * The second hit-test is required to allow the click to be "cancelled" by moving away
+ * from the target. This is standard GUI behavior. You can disable this second hit test
+ * using the .HitTestOnRelease property. This is strongly discouraged.
+ *    
+ * The hit-test and on-clicked functions are provided by a IClickBehaviorTarget instance.
+ */
+UCLASS()
+class INTERACTIVETOOLSFRAMEWORK_API USingleClickToolBehavior : public UAnyButtonInputBehavior
+{
+	GENERATED_BODY()
+
+public:
+	USingleClickToolBehavior();
+
+	/**
+	 * Initialize this behavior with the given Target
+	 * @param Target implementor of hit-test and on-clicked functions
+	 */
+	virtual void Initialize(IClickBehaviorTarget* Target);
+
+
+	// UInputBehavior implementation
+
+	virtual FInputCaptureRequest WantsCapture(const FInputDeviceState& Input) override;
+	virtual FInputCaptureUpdate BeginCapture(const FInputDeviceState& Input, EInputCaptureSide eSide) override;
+	virtual FInputCaptureUpdate UpdateCapture(const FInputDeviceState& Input, const FInputCaptureData& Data) override;
+	virtual void ForceEndCapture(const FInputCaptureData& Data) override;
+
+
+public:
+	/** Hit-test is repeated on release (standard behavior). If false, */
+	UPROPERTY()
+	bool HitTestOnRelease;
+
+
+protected:
+	/** Click Target object */
+	IClickBehaviorTarget* Target;
+
+	/**
+	 * Internal function that forwards click evens to Target::OnClicked, you can customize behavior here
+	 */
+	virtual void Clicked(const FInputDeviceState& Input, const FInputCaptureData& Data);
+};
+
diff --git a/Engine/Source/Runtime/InteractiveToolsFramework/Public/BaseTools/MeshSurfacePointTool.h b/Engine/Source/Runtime/InteractiveToolsFramework/Public/BaseTools/MeshSurfacePointTool.h
new file mode 100644
--- /dev/null
+++ b/Engine/Source/Runtime/InteractiveToolsFramework/Public/BaseTools/MeshSurfacePointTool.h
@@ -0,0 +1,144 @@
+// Copyright 1998-2019 Epic Games, Inc. All Rights Reserved.
+
+#pragma once
+
+#include "CoreMinimal.h"
+#include "InteractiveTool.h"
+#include "InteractiveToolBuilder.h"
+#include "InputBehaviorSet.h"
+#include "BaseBehaviors/BehaviorTargetInterfaces.h"
+#include "BaseBehaviors/AnyButtonInputBehavior.h"
+#include "MeshSurfacePointTool.generated.h"
+
+
+
+class UMeshSurfacePointTool;
+
+/**
+ * 
+ */
+UCLASS()
+class INTERACTIVETOOLSFRAMEWORK_API UMeshSurfacePointToolBuilder : public UInteractiveToolBuilder
+{
+	GENERATED_BODY()
+
+public:
+	/** @return true if a single mesh source can be found in the active selection */
+	virtual bool CanBuildTool(const FToolBuilderState& SceneState) const override;
+
+	/** @return new Tool instance initialized with selected mesh source */
+	virtual UInteractiveTool* BuildTool(const FToolBuilderState& SceneState) const override;
+
+	/** @return new Tool instance. Override this in subclasses to build a different Tool class type */
+	virtual UMeshSurfacePointTool* CreateNewTool(const FToolBuilderState& SceneState) const;
+
+	/** Called by BuildTool to configure the Tool with the input MeshSource based on the SceneState */
+	virtual void InitializeNewTool(UMeshSurfacePointTool* Tool, const FToolBuilderState& SceneState) const;
+};
+
+
+
+/**
+ * UMeshSurfacePointTool is a base Tool implementation that can be used to implement various
+ * "point on surface" interactions. The tool acts on an input IMeshDescriptionSource object,
+ * which the standard Builder can extract from the current selection (eg Editor selection).
+ * 
+ * Subclasses override the OnBeginDrag/OnUpdateDrag/OnEndDrag and OnUpdateHover functions
+ * to implement custom behavior.
+ */
+UCLASS()
+class INTERACTIVETOOLSFRAMEWORK_API UMeshSurfacePointTool : public UInteractiveTool, public IHoverBehaviorTarget
+{
+	GENERATED_BODY()
+
+public:
+	// UInteractiveTool API Implementation
+
+	/** Register InputBehaviors, etc */
+	virtual void Setup() override;
+
+
+	// UMeshSurfacePointTool API
+
+	/**
+	 * Set the Target MeshDescription-providing Object for this tool
+	 * @param MeshSource object that can provide a MeshDescriptionSource. Note UMeshSurfacePointTool
+	 *        only calls MeshSource->HitTest, so this source technically does not have to actually provide
+	 *        a valid MeshDescription
+	 */
+	virtual void SetSelection(TUniquePtr<IMeshDescriptionSource> MeshSource);
+
+	/**
+	 * @return true if the target MeshSource is hit by the Ray
+	 */
+	virtual bool HitTest(const FRay& Ray, FHitResult& OutHit);
+
+	
+	/**
+	 * This function is called by registered InputBehaviors when the user begins a click-drag-release interaction
+	 */
+	virtual void OnBeginDrag(const FRay& Ray);
+
+	/**
+	 * This function is called by registered InputBehaviorseach frame that the user is in a click-drag-release interaction
+	 */
+	virtual void OnUpdateDrag(const FRay& Ray);
+
+	/**
+	 * This function is called by registered InputBehaviors when the user releases the button driving a click-drag-release interaction
+	 */
+	virtual void OnEndDrag(const FRay& Ray);
+
+	/**
+	 * This function is called each frame while the tool is active, to support hover updates
+	 */
+	virtual void OnUpdateHover(const FInputDeviceRay& DevicePos) {}
+
+
+
+	/** Called by registered InputBehaviors to set the state of the "shift" button (or device equivalent) */
+	virtual void SetShiftToggle(bool bShiftDown);
+
+	/** @return current state of the shift toggle */
+	virtual bool GetShiftToggle() { return bShiftToggle; }
+
+
+protected:
+	/** Target source object that provides a MeshDescription and various other query API calls */
+	TUniquePtr<IMeshDescriptionSource> MeshSource;
+
+	/** Current state of the shift modifier toggle */
+	bool bShiftToggle;
+};
+
+
+
+
+
+
+
+
+/**
+ * UMeshSurfacePointToolMouseBehavior implements mouse press-drag-release interaction behavior for Mouse devices.
+ * You can configure the base UAnyButtonInputBehavior to change the mouse button in use (default = left mouse)
+ */
+UCLASS()
+class INTERACTIVETOOLSFRAMEWORK_API UMeshSurfacePointToolMouseBehavior : public UAnyButtonInputBehavior
+{
+	GENERATED_BODY()
+
+public:
+	virtual void Initialize(UMeshSurfacePointTool* Tool);
+
+	virtual FInputCaptureRequest WantsCapture(const FInputDeviceState& input) override;
+	virtual FInputCaptureUpdate BeginCapture(const FInputDeviceState& input, EInputCaptureSide eSide) override;
+	virtual FInputCaptureUpdate UpdateCapture(const FInputDeviceState& input, const FInputCaptureData& data) override;
+	virtual void ForceEndCapture(const FInputCaptureData& data) override;
+
+protected:
+	UMeshSurfacePointTool* Tool;
+	FRay LastWorldRay;
+	bool bInDragCapture;
+};
+
+
diff --git a/Engine/Source/Runtime/InteractiveToolsFramework/Public/BaseTools/SingleClickTool.h b/Engine/Source/Runtime/InteractiveToolsFramework/Public/BaseTools/SingleClickTool.h
new file mode 100644
--- /dev/null
+++ b/Engine/Source/Runtime/InteractiveToolsFramework/Public/BaseTools/SingleClickTool.h
@@ -0,0 +1,62 @@
+// Copyright 1998-2019 Epic Games, Inc. All Rights Reserved.
+
+#pragma once
+
+#include "CoreMinimal.h"
+#include "InteractiveTool.h"
+#include "InteractiveToolBuilder.h"
+#include "BaseBehaviors/SingleClickBehavior.h"
+#include "SingleClickTool.generated.h"
+
+
+
+/**
+ * Builder for USingleClickTool
+ */
+UCLASS()
+class INTERACTIVETOOLSFRAMEWORK_API USingleClickToolBuilder : public UInteractiveToolBuilder
+{
+	GENERATED_BODY()
+
+public:
+	virtual bool CanBuildTool(const FToolBuilderState& SceneState) const override;
+	virtual UInteractiveTool* BuildTool(const FToolBuilderState& SceneState) const override;
+};
+
+
+/**
+ * USingleClickTool is perhaps the simplest possible interactive tool. It simply
+ * reacts to default primary button clicks for the active device (eg left-mouse clicks).
+ * 
+ * The function ::IsHitByClick() determines what is clickable by this Tool. The default is
+ * to return true, which means the click will activate anywhere (the Tool itself has no
+ * notion of Actors, Components, etc). You can override this function to, for example,
+ * filter out clicks that don't hit a target object, etc.
+ * 
+ * The function ::OnClicked() implements the action that will occur when a click happens.
+ * You must override this to implement any kind of useful behavior.
+ */
+UCLASS()
+class INTERACTIVETOOLSFRAMEWORK_API USingleClickTool : public UInteractiveTool, public IClickBehaviorTarget
+{
+	GENERATED_BODY()
+
+public:
+
+	/**
+	 * Register default primary-button-click InputBehaviors
+	 */
+	virtual void Setup() override;
+
+	/**
+	 * Test if the Target is hit at this 2D position / 3D ray
+	 */
+	virtual bool IsHitByClick(const FInputDeviceRay& ClickPos);
+
+	/**
+	 * Click the Target at this 2D position / 3D ray. Default behavior is to print debug string.
+	 */
+	virtual void OnClicked(const FInputDeviceRay& ClickPos);
+
+};
+
