minor utility-function additions to InteractiveToolsFramework
#rnx

diff --git a/Engine/Source/Runtime/InteractiveToolsFramework/Private/InputBehavior.cpp b/Engine/Source/Runtime/InteractiveToolsFramework/Private/InputBehavior.cpp
--- a/Engine/Source/Runtime/InteractiveToolsFramework/Private/InputBehavior.cpp
+++ b/Engine/Source/Runtime/InteractiveToolsFramework/Private/InputBehavior.cpp
@@ -28,6 +28,11 @@ FInputCapturePriority UInputBehavior::GetPriority()
 	return DefaultPriority;
 }
 
+void UInputBehavior::SetDefaultPriority(const FInputCapturePriority& Priority)
+{
+	DefaultPriority = Priority;
+}
+
 EInputDevices UInputBehavior::GetSupportedDevices()
 {
 	return EInputDevices::Mouse;
diff --git a/Engine/Source/Runtime/InteractiveToolsFramework/Private/ToolBuilderUtil.cpp b/Engine/Source/Runtime/InteractiveToolsFramework/Private/ToolBuilderUtil.cpp
--- a/Engine/Source/Runtime/InteractiveToolsFramework/Private/ToolBuilderUtil.cpp
+++ b/Engine/Source/Runtime/InteractiveToolsFramework/Private/ToolBuilderUtil.cpp
@@ -91,3 +91,38 @@ UObject* ToolBuilderUtil::FindFirstComponent(const FToolBuilderState& InputState
 
 
 
+
+TArray<UObject*> ToolBuilderUtil::FindAllComponents(const FToolBuilderState& InputState, const TFunction<bool(UObject*)>& Predicate)
+{
+	TArray<UObject*> Components;
+
+	if (InputState.SelectedComponents != nullptr && InputState.SelectedComponents->Num() > 0)
+	{
+		for (FSelectionIterator Iter(*InputState.SelectedComponents); Iter; ++Iter)
+		{
+			if (Predicate(*Iter))
+			{
+				Components.AddUnique(*Iter);
+			}
+		}
+	}
+	else
+	{
+		for (FSelectionIterator Iter(*InputState.SelectedActors); Iter; ++Iter)
+		{
+			if (AActor* Actor = Cast<AActor>(*Iter))
+			{
+				for (UActorComponent* Component : Actor->GetComponents())
+				{
+					if (Predicate(Component))
+					{
+						Components.AddUnique(Component);
+					}
+				}
+			}
+		}
+	}
+
+	return Components;
+}
+
diff --git a/Engine/Source/Runtime/InteractiveToolsFramework/Public/InputBehavior.h b/Engine/Source/Runtime/InteractiveToolsFramework/Public/InputBehavior.h
--- a/Engine/Source/Runtime/InteractiveToolsFramework/Public/InputBehavior.h
+++ b/Engine/Source/Runtime/InteractiveToolsFramework/Public/InputBehavior.h
@@ -188,14 +188,29 @@ struct INTERACTIVETOOLSFRAMEWORK_API FInputCaptureUpdate
  */
 struct INTERACTIVETOOLSFRAMEWORK_API FInputCapturePriority
 {
+	static constexpr int DEFAULT_GIZMO_PRIORITY = 50;
+	static constexpr int DEFAULT_TOOL_PRIORITY = 100;
+
 	/** Constant priority value */
 	int Priority;
 
-	FInputCapturePriority(int priority = 100) 
+	FInputCapturePriority(int priority = DEFAULT_TOOL_PRIORITY)
 	{
 		Priority = priority;
 	}
 
+	/** @return a priority lower than this priority */
+	FInputCapturePriority MakeLower(int DeltaAmount = 1) const
+	{
+		return FInputCapturePriority(Priority + DeltaAmount);
+	}
+
+	/** @return a priority higher than this priority */
+	FInputCapturePriority MakeHigher(int DeltaAmount = 1) const
+	{
+		return FInputCapturePriority(Priority - DeltaAmount);
+	}
+
 	friend bool operator<(const FInputCapturePriority& l, const FInputCapturePriority& r)
 	{
 		return l.Priority < r.Priority;
@@ -243,6 +258,10 @@ public:
 	/** The priority is used to resolve situations where multiple behaviors want the same capture */
 	virtual FInputCapturePriority GetPriority();
 
+	/** Configure the default priority of an instance of this behavior */
+	virtual void SetDefaultPriority(const FInputCapturePriority& Priority);
+
+
 	/** Which device types does this Behavior support */
 	virtual EInputDevices GetSupportedDevices();
 
diff --git a/Engine/Source/Runtime/InteractiveToolsFramework/Public/InputState.h b/Engine/Source/Runtime/InteractiveToolsFramework/Public/InputState.h
--- a/Engine/Source/Runtime/InteractiveToolsFramework/Public/InputState.h
+++ b/Engine/Source/Runtime/InteractiveToolsFramework/Public/InputState.h
@@ -184,5 +184,36 @@ struct FInputDeviceState
 	{
 		return ((InputDevice & DeviceType) != EInputDevices::None);
 	}
+
+
+
+
+	//
+	// utility functions to pass as lambdas
+	// 
+
+	/** @return true if shift key is down in input state */
+	static bool IsShiftKeyDown(const FInputDeviceState& InputState)
+	{
+		return InputState.bShiftKeyDown;
+	}
+
+	/** @return true if ctrl key is down in input state */
+	static bool IsCtrlKeyDown(const FInputDeviceState& InputState)
+	{
+		return InputState.bCtrlKeyDown;
+	}
+
+	/** @return true if alt key is down in input state */
+	static bool IsAltKeyDown(const FInputDeviceState& InputState)
+	{
+		return InputState.bAltKeyDown;
+	}
+
+	/** @return true if Apple Command key is down in input state */
+	static bool IsCmdKeyDown(const FInputDeviceState& InputState)
+	{
+		return InputState.bCmdKeyDown;
+	}
 };
 
diff --git a/Engine/Source/Runtime/InteractiveToolsFramework/Public/ToolBuilderUtil.h b/Engine/Source/Runtime/InteractiveToolsFramework/Public/ToolBuilderUtil.h
--- a/Engine/Source/Runtime/InteractiveToolsFramework/Public/ToolBuilderUtil.h
+++ b/Engine/Source/Runtime/InteractiveToolsFramework/Public/ToolBuilderUtil.h
@@ -24,6 +24,10 @@ namespace ToolBuilderUtil
 	INTERACTIVETOOLSFRAMEWORK_API
 	UObject* FindFirstComponent(const FToolBuilderState& InputState, const TFunction<bool(UObject*)>& Predicate);
 
+	/** First all components that passes predicate. Searches Components selection list first, then all Actors */
+	INTERACTIVETOOLSFRAMEWORK_API
+	TArray<UObject*> FindAllComponents(const FToolBuilderState& InputState, const TFunction<bool(UObject*)>& Predicate);
+
 	// @todo not sure that actors with multiple components are handled properly...
 	/** Count number of components of given type. If Component selection is not empty, returns that count, otherwise counts in all selected Actors */
 	template<typename ComponentType>
