add/update base behaviors and tools - MultiClickSequenceInputBehavior, IModifierToggleBehaviorTarget, FInputBehaviorModifierStates
#rnx

diff --git a/Engine/Source/Runtime/InteractiveToolsFramework/Private/BaseBehaviors/MultiClickSequenceInputBehavior.cpp b/Engine/Source/Runtime/InteractiveToolsFramework/Private/BaseBehaviors/MultiClickSequenceInputBehavior.cpp
new file mode 100644
--- /dev/null
+++ b/Engine/Source/Runtime/InteractiveToolsFramework/Private/BaseBehaviors/MultiClickSequenceInputBehavior.cpp
@@ -0,0 +1,75 @@
+// Copyright 1998-2019 Epic Games, Inc. All Rights Reserved.
+
+#include "BaseBehaviors/MultiClickSequenceInputBehavior.h"
+
+
+
+UMultiClickSequenceInputBehavior::UMultiClickSequenceInputBehavior()
+{
+	bInActiveSequence = false;
+}
+
+
+void UMultiClickSequenceInputBehavior::Initialize(IClickSequenceBehaviorTarget* TargetIn)
+{
+	this->Target = TargetIn;
+	bInActiveSequence = false;
+}
+
+
+FInputCaptureRequest UMultiClickSequenceInputBehavior::WantsCapture(const FInputDeviceState& input)
+{
+	check(bInActiveSequence == false);   // should not happen...
+	bInActiveSequence = false;
+
+	if ( IsPressed(input) && (ModifierCheckFunc == nullptr || ModifierCheckFunc(input) ) )
+	{
+		if ( Target->CanBeginClickSequence(GetDeviceRay(input)) )
+		{
+			return FInputCaptureRequest::Begin(this, EInputCaptureSide::Any);
+		}
+	}
+	return FInputCaptureRequest::Ignore();
+}
+
+
+FInputCaptureUpdate UMultiClickSequenceInputBehavior::BeginCapture(const FInputDeviceState& input, EInputCaptureSide eSide)
+{
+	Modifiers.UpdateModifiers(input, Target);
+
+	Target->OnBeginClickSequence(GetDeviceRay(input));
+	bInActiveSequence = true;
+	
+	return FInputCaptureUpdate::Begin(this, EInputCaptureSide::Any);
+}
+
+
+FInputCaptureUpdate UMultiClickSequenceInputBehavior::UpdateCapture(const FInputDeviceState& input, const FInputCaptureData& data)
+{
+	check(bInActiveSequence == true);   // should always be the case!
+
+	Modifiers.UpdateModifiers(input, Target);
+
+	if (IsReleased(input)) 
+	{
+		bool bContinue = Target->OnNextSequenceClick(GetDeviceRay(input));
+		if (bContinue == false)
+		{
+			bInActiveSequence = false;
+			return FInputCaptureUpdate::End();
+		}
+	}
+	else
+	{
+		Target->OnNextSequencePreview(GetDeviceRay(input));
+	}
+
+	return FInputCaptureUpdate::Continue();
+}
+
+
+void UMultiClickSequenceInputBehavior::ForceEndCapture(const FInputCaptureData& data)
+{
+	Target->OnTerminateClickSequence();
+	bInActiveSequence = false;
+}
diff --git a/Engine/Source/Runtime/InteractiveToolsFramework/Private/BaseBehaviors/SingleClickBehavior.cpp b/Engine/Source/Runtime/InteractiveToolsFramework/Private/BaseBehaviors/SingleClickBehavior.cpp
--- a/Engine/Source/Runtime/InteractiveToolsFramework/Private/BaseBehaviors/SingleClickBehavior.cpp
+++ b/Engine/Source/Runtime/InteractiveToolsFramework/Private/BaseBehaviors/SingleClickBehavior.cpp
@@ -4,21 +4,21 @@
 
 
 
-USingleClickToolBehavior::USingleClickToolBehavior()
+USingleClickInputBehavior::USingleClickInputBehavior()
 {
 	HitTestOnRelease = true;
 }
 
 
-void USingleClickToolBehavior::Initialize(IClickBehaviorTarget* TargetIn)
+void USingleClickInputBehavior::Initialize(IClickBehaviorTarget* TargetIn)
 {
 	this->Target = TargetIn;
 }
 
 
-FInputCaptureRequest USingleClickToolBehavior::WantsCapture(const FInputDeviceState& input)
+FInputCaptureRequest USingleClickInputBehavior::WantsCapture(const FInputDeviceState& input)
 {
-	if (IsPressed(input))
+	if (IsPressed(input) && (ModifierCheckFunc == nullptr || ModifierCheckFunc(input)) )
 	{
 		if ( Target->IsHitByClick(GetDeviceRay(input)) )
 		{
@@ -29,13 +29,13 @@ FInputCaptureRequest USingleClickToolBehavior::WantsCapture(const FInputDeviceSt
 }
 
 
-FInputCaptureUpdate USingleClickToolBehavior::BeginCapture(const FInputDeviceState& input, EInputCaptureSide eSide)
+FInputCaptureUpdate USingleClickInputBehavior::BeginCapture(const FInputDeviceState& input, EInputCaptureSide eSide)
 {
 	return FInputCaptureUpdate::Begin(this, EInputCaptureSide::Any);
 }
 
 
-FInputCaptureUpdate USingleClickToolBehavior::UpdateCapture(const FInputDeviceState& input, const FInputCaptureData& data)
+FInputCaptureUpdate USingleClickInputBehavior::UpdateCapture(const FInputDeviceState& input, const FInputCaptureData& data)
 {
 	if (IsReleased(input)) 
 	{
@@ -52,13 +52,13 @@ FInputCaptureUpdate USingleClickToolBehavior::UpdateCapture(const FInputDeviceSt
 }
 
 
-void USingleClickToolBehavior::ForceEndCapture(const FInputCaptureData& data)
+void USingleClickInputBehavior::ForceEndCapture(const FInputCaptureData& data)
 {
 	// nothing to do
 }
 
 
-void USingleClickToolBehavior::Clicked(const FInputDeviceState& input, const FInputCaptureData& data)
+void USingleClickInputBehavior::Clicked(const FInputDeviceState& input, const FInputCaptureData& data)
 {
 	Target->OnClicked(GetDeviceRay(input));
 }
diff --git a/Engine/Source/Runtime/InteractiveToolsFramework/Private/BaseTools/SingleClickTool.cpp b/Engine/Source/Runtime/InteractiveToolsFramework/Private/BaseTools/SingleClickTool.cpp
--- a/Engine/Source/Runtime/InteractiveToolsFramework/Private/BaseTools/SingleClickTool.cpp
+++ b/Engine/Source/Runtime/InteractiveToolsFramework/Private/BaseTools/SingleClickTool.cpp
@@ -31,7 +31,7 @@ void USingleClickTool::Setup()
 	UInteractiveTool::Setup();
 
 	// add default button input behaviors for devices
-	USingleClickToolBehavior* MouseBehavior = NewObject<USingleClickToolBehavior>();
+	USingleClickInputBehavior* MouseBehavior = NewObject<USingleClickInputBehavior>();
 	MouseBehavior->Initialize(this);
 	AddInputBehavior(MouseBehavior);
 }
diff --git a/Engine/Source/Runtime/InteractiveToolsFramework/Public/BaseBehaviors/BehaviorTargetInterfaces.h b/Engine/Source/Runtime/InteractiveToolsFramework/Public/BaseBehaviors/BehaviorTargetInterfaces.h
--- a/Engine/Source/Runtime/InteractiveToolsFramework/Public/BaseBehaviors/BehaviorTargetInterfaces.h
+++ b/Engine/Source/Runtime/InteractiveToolsFramework/Public/BaseBehaviors/BehaviorTargetInterfaces.h
@@ -4,11 +4,33 @@
 
 #include "CoreMinimal.h"
 
+
+/**
+ * IModifierToggleBehaviorTarget is an interface that InputBehaviors can use to notify
+ * a target about modifier toggle states (like shift key being down, etc).
+ * The meaning of the modifier ID is client-defined (generally provided to the InputBehavior in a registration step)
+ */
+class IModifierToggleBehaviorTarget
+{
+public:
+	virtual ~IModifierToggleBehaviorTarget() {}
+
+	/**
+	 * Notify target of current modifier state
+	 * @param ModifierID client-defined integer that identifiers modifier
+	 * @param bIsOn is modifier current on or off
+	 */
+	virtual void OnUpdateModifierState(int ModifierID, bool bIsOn)
+	{
+	}
+};
+
+
 /**
  * Functions required to apply standard "Click" state machines to a target object.
  * See USingleClickBehavior for an example of this kind of state machine.
  */
-class IClickBehaviorTarget
+class IClickBehaviorTarget : public IModifierToggleBehaviorTarget
 {
 public:
 	virtual ~IClickBehaviorTarget() {}
@@ -30,6 +52,53 @@ public:
 
 
 
+
+/**
+ * Target interface used by InputBehaviors that want to implement a multi-click sequence
+ * (eg such as drawing a polygon with multiple clicks)
+ */
+class IClickSequenceBehaviorTarget : public IModifierToggleBehaviorTarget
+{
+public:
+	virtual ~IClickSequenceBehaviorTarget() {}
+
+	/**
+	 * Test if target would like to begin sequence based on this click
+	 * @param ClickPos device position/ray at click point
+	 * @return true if target wants to begin sequence
+	 */
+	virtual bool CanBeginClickSequence(const FInputDeviceRay& ClickPos) = 0;
+
+	/**
+	 * Notify Target that click sequence can begin at click point
+	 * @param ClickPos device position/ray at click point
+	 */
+	virtual void OnBeginClickSequence(const FInputDeviceRay& ClickPos) = 0;
+	
+	/**
+	 * Notify Target device position has changed but a click hasn't ocurred yet (eg for interactive previews)
+	 * @param ClickPos device position/ray at click point
+	 */
+	virtual void OnNextSequencePreview(const FInputDeviceRay& ClickPos) = 0;
+
+	/**
+	 * Notify Target about next click in sqeuence
+	 * @param ClickPos device position/ray at click point
+	 * @return false if Target wants to terminate sequence
+	 */
+	virtual bool OnNextSequenceClick(const FInputDeviceRay& ClickPos) = 0;
+
+	/**
+	 * Notify Target that click sequence has been explicitly terminated (eg by escape key, cancel tool, etc)
+	 */
+	virtual void OnTerminateClickSequence() = 0;
+
+};
+
+
+
+
+
 /**
  * IHoverBehaviorTarget allows Behaviors to notify Tools/etc about
  * device event data in a generic way, without requiring that all Tools
diff --git a/Engine/Source/Runtime/InteractiveToolsFramework/Public/BaseBehaviors/InputBehaviorModifierStates.h b/Engine/Source/Runtime/InteractiveToolsFramework/Public/BaseBehaviors/InputBehaviorModifierStates.h
new file mode 100644
--- /dev/null
+++ b/Engine/Source/Runtime/InteractiveToolsFramework/Public/BaseBehaviors/InputBehaviorModifierStates.h
@@ -0,0 +1,57 @@
+// Copyright 1998-2019 Epic Games, Inc. All Rights Reserved.
+
+#pragma once
+
+#include "CoreMinimal.h"
+#include "InputState.h"
+#include "BehaviorTargetInterfaces.h"
+
+/**
+ * FInputBehaviorModifierStates is an object that can be placed in an InputBehavior to allow
+ * users of the behavior to request that they be notified about modifier keys/buttons/etc state.
+ * 
+ * We don't know ahead of time what might be used to determine modifier state, which input
+ * devices might be used, etc. So the user has to register (ModifierID,ModifierTestFunction) pairs.
+ * The behavior then calls UpdateModifiers() which will query each of the test functions and
+ * notify the target object about the state of the modifier.
+ */
+class FInputBehaviorModifierStates
+{
+public:
+	typedef TFunction<bool(const FInputDeviceState&)> FModifierTestFunction;
+
+protected:
+	/** List of modifier IDs that have been registered */
+	TArray<int> ModifierIDs;
+	/** The modifier test function for each ID */
+	TMap<int, FModifierTestFunction> ModifierTests;
+
+public:
+	
+	/**
+	 * Register a modifier ID and an associated test function
+	 * @param ModifierID the modifier ID
+	 * @param ModifierTest the test function
+	 */
+	void RegisterModifier(int ModifierID, const FModifierTestFunction& ModifierTest)
+	{
+		check(ModifierIDs.Contains(ModifierID) == false);
+		ModifierIDs.Add(ModifierID);
+		ModifierTests.Add(ModifierID, ModifierTest);
+	}
+
+
+	/**
+	 * Look up the current state of each registered modifier and pass to the target
+	 * @param Input the current input device state
+	 * @param ModifiersTarget the target that is interested in the modifier values
+	 */
+	void UpdateModifiers(const FInputDeviceState& Input, IModifierToggleBehaviorTarget* ModifiersTarget) const
+	{
+		for (int ModifierID : ModifierIDs)
+		{
+			bool bIsOn = ModifierTests[ModifierID](Input);
+			ModifiersTarget->OnUpdateModifierState(ModifierID, bIsOn);
+		}
+	}
+};
\ No newline at end of file
diff --git a/Engine/Source/Runtime/InteractiveToolsFramework/Public/BaseBehaviors/MultiClickSequenceInputBehavior.h b/Engine/Source/Runtime/InteractiveToolsFramework/Public/BaseBehaviors/MultiClickSequenceInputBehavior.h
new file mode 100644
--- /dev/null
+++ b/Engine/Source/Runtime/InteractiveToolsFramework/Public/BaseBehaviors/MultiClickSequenceInputBehavior.h
@@ -0,0 +1,70 @@
+// Copyright 1998-2019 Epic Games, Inc. All Rights Reserved.
+
+#pragma once
+
+#include "CoreMinimal.h"
+#include "BaseBehaviors/AnyButtonInputBehavior.h"
+#include "BehaviorTargetInterfaces.h"
+#include "BaseBehaviors/InputBehaviorModifierStates.h"
+#include "MultiClickSequenceInputBehavior.generated.h"
+
+
+/**
+ * UMultiClickSequenceInputBehavior implements a generic multi-click-sequence input behavior.
+ * For example this behavior could be used to implement a multi-click polygon-drawing interaction.
+ *
+ * The internal state machine works as follows:
+ *   1) on input-device-button-press, check if target wants to begin sequence. If so, begin capture.
+ *   2) on button *release*, check if target wants to continue or terminate sequence
+ *       a) if terminate, release capture
+ *       b) if continue, do nothing (capture continues between presses)
+ *
+ * The target will receive "preview" notifications (basically hover) during updates where there is
+ * not a release. This can be used to (eg) update a rubber-band selection end point
+ * 
+ * @todo it may be better to implement this as multiple captures, and use hover callbacks to 
+ * do the between-capture previews. holding capture across mouse release is not ideal.
+ */
+UCLASS()
+class INTERACTIVETOOLSFRAMEWORK_API UMultiClickSequenceInputBehavior : public UAnyButtonInputBehavior
+{
+	GENERATED_BODY()
+
+public:
+	UMultiClickSequenceInputBehavior();
+
+	/**
+	 * The modifier set for this behavior
+	 */
+	FInputBehaviorModifierStates Modifiers;
+
+	/**
+	 * Initialize this behavior with the given Target
+	 * @param Target implementor of hit-test and on-clicked functions
+	 */
+	virtual void Initialize(IClickSequenceBehaviorTarget* Target);
+
+
+	/**
+	 * WantsCapture() will only return capture request if this function returns true (or is null)
+	 */
+	TFunction<bool(const FInputDeviceState&)> ModifierCheckFunc = nullptr;
+
+
+	// UInputBehavior implementation
+
+	virtual FInputCaptureRequest WantsCapture(const FInputDeviceState& Input) override;
+	virtual FInputCaptureUpdate BeginCapture(const FInputDeviceState& Input, EInputCaptureSide eSide) override;
+	virtual FInputCaptureUpdate UpdateCapture(const FInputDeviceState& Input, const FInputCaptureData& Data) override;
+	virtual void ForceEndCapture(const FInputCaptureData& Data) override;
+
+
+public:
+
+protected:
+	/** Click Target object */
+	IClickSequenceBehaviorTarget* Target;
+
+	bool bInActiveSequence;
+};
+
diff --git a/Engine/Source/Runtime/InteractiveToolsFramework/Public/BaseBehaviors/SingleClickBehavior.h b/Engine/Source/Runtime/InteractiveToolsFramework/Public/BaseBehaviors/SingleClickBehavior.h
--- a/Engine/Source/Runtime/InteractiveToolsFramework/Public/BaseBehaviors/SingleClickBehavior.h
+++ b/Engine/Source/Runtime/InteractiveToolsFramework/Public/BaseBehaviors/SingleClickBehavior.h
@@ -9,7 +9,7 @@
 
 
 /**
- * USingleClickToolBehavior implements a standard "button-click"-style input behavior.
+ * USingleClickInputBehavior implements a standard "button-click"-style input behavior.
  * The state machine works as follows:
  *    1) on input-device-button-press, hit-test the target. If hit, begin capture
  *    2) on input-device-button-release, hit-test the target. If hit, call Target::OnClicked(). If not hit, ignore click.
@@ -21,12 +21,12 @@
  * The hit-test and on-clicked functions are provided by a IClickBehaviorTarget instance.
  */
 UCLASS()
-class INTERACTIVETOOLSFRAMEWORK_API USingleClickToolBehavior : public UAnyButtonInputBehavior
+class INTERACTIVETOOLSFRAMEWORK_API USingleClickInputBehavior : public UAnyButtonInputBehavior
 {
 	GENERATED_BODY()
 
 public:
-	USingleClickToolBehavior();
+	USingleClickInputBehavior();
 
 	/**
 	 * Initialize this behavior with the given Target
@@ -35,6 +35,12 @@ public:
 	virtual void Initialize(IClickBehaviorTarget* Target);
 
 
+	/**
+	 * WantsCapture() will only return capture request if this function returns true (or is null)
+	 */
+	TFunction<bool(const FInputDeviceState&)> ModifierCheckFunc = nullptr;
+
+
 	// UInputBehavior implementation
 
 	virtual FInputCaptureRequest WantsCapture(const FInputDeviceState& Input) override;
