warning fix
#rnx

diff --git a/Engine/Source/Runtime/InteractiveToolsFramework/Private/InteractiveTool.cpp b/Engine/Source/Runtime/InteractiveToolsFramework/Private/InteractiveTool.cpp
--- a/Engine/Source/Runtime/InteractiveToolsFramework/Private/InteractiveTool.cpp
+++ b/Engine/Source/Runtime/InteractiveToolsFramework/Private/InteractiveTool.cpp
@@ -52,9 +52,9 @@ void UInteractiveTool::AddToolPropertySource(UInteractiveToolPropertySet* Proper
 	check(ToolPropertyObjects.Contains(PropertySet) == false);
 	ToolPropertyObjects.Add(PropertySet);
 	// @todo do we need to create a lambda every time for this?
-	PropertySet->GetOnModified().AddLambda([this](UObject* PropertySet, UProperty* Property)
+	PropertySet->GetOnModified().AddLambda([this](UObject* PropertySetArg, UProperty* PropertyArg)
 	{
-		OnPropertyModified(PropertySet, Property);
+		OnPropertyModified(PropertySetArg, PropertyArg);
 	});
 }
 
