Add InteractiveGizmo and InteractiveGizmoManager
#rnx

diff --git a/Engine/Source/Editor/EditorInteractiveToolsFramework/Private/EdModeInteractiveToolsContext.cpp b/Engine/Source/Editor/EditorInteractiveToolsFramework/Private/EdModeInteractiveToolsContext.cpp
--- a/Engine/Source/Editor/EditorInteractiveToolsFramework/Private/EdModeInteractiveToolsContext.cpp
+++ b/Engine/Source/Editor/EditorInteractiveToolsFramework/Private/EdModeInteractiveToolsContext.cpp
@@ -32,6 +32,7 @@ public:
 	virtual void GetCurrentSelectionState(FToolBuilderState& StateOut) const override
 	{
 		StateOut.ToolManager = ToolsContext->ToolManager;
+		StateOut.GizmoManager = ToolsContext->GizmoManager;
 		StateOut.World = EditorMode->GetWorld();
 		StateOut.SelectedActors = EditorMode->GetModeManager()->GetSelectedActors();
 		StateOut.SelectedComponents = EditorMode->GetModeManager()->GetSelectedComponents();
@@ -176,6 +177,7 @@ void UEdModeInteractiveToolsContext::PostInvalidation()
 void UEdModeInteractiveToolsContext::Tick(FEditorViewportClient* ViewportClient, float DeltaTime)
 {
 	ToolManager->Tick(DeltaTime);
+	GizmoManager->Tick(DeltaTime);
 
 	if (bInvalidationPending)
 	{
@@ -203,6 +205,7 @@ void UEdModeInteractiveToolsContext::Render(const FSceneView* View, FViewport* V
 	TempRenderContext RenderContext;
 	RenderContext.PDI = PDI;
 	ToolManager->Render(&RenderContext);
+	GizmoManager->Render(&RenderContext);
 }
 
 
@@ -308,10 +311,21 @@ bool UEdModeInteractiveToolsContext::MouseMove(FEditorViewportClient* ViewportCl
 	FInputDeviceState InputState = CurrentInputState;
 	InputState.InputDevice = EInputDevices::Mouse;
 
-	if (InputRouter->HasActiveMouseCapture())
-		UE_LOG(LogTemp, Warning, TEXT("MOUSE MOVE DURING CAPTURE?"));
+	InputState.SetKeyStates(
+		ViewportClient->IsShiftPressed(), ViewportClient->IsAltPressed(),
+		ViewportClient->IsCtrlPressed(), ViewportClient->IsCmdPressed());
 
-	InputRouter->PostHoverInputEvent(InputState);
+	if (InputRouter->HasActiveMouseCapture())
+	{
+		// This state occurs if InputBehavior did not release capture on mouse release. 
+		// UMultiClickSequenceInputBehavior does this, eg for multi-click draw-polygon sequences.
+		// It's not ideal though and maybe would be better done via multiple captures + hover...?
+		InputRouter->PostInputEvent(InputState);
+	}
+	else
+	{
+		InputRouter->PostHoverInputEvent(InputState);
+	}
 
 	return false;
 }
diff --git a/Engine/Source/Runtime/InteractiveToolsFramework/Private/InteractiveGizmo.cpp b/Engine/Source/Runtime/InteractiveToolsFramework/Private/InteractiveGizmo.cpp
new file mode 100644
--- /dev/null
+++ b/Engine/Source/Runtime/InteractiveToolsFramework/Private/InteractiveGizmo.cpp
@@ -0,0 +1,52 @@
+// Copyright 1998-2019 Epic Games, Inc. All Rights Reserved.
+
+
+#include "InteractiveGizmo.h"
+#include "InteractiveGizmoManager.h"
+
+
+UInteractiveGizmo::UInteractiveGizmo()
+{
+	// Gizmos don't get saved but this isn't necessary because they are created in the transient package...
+	SetFlags(RF_Transient);
+
+	InputBehaviors = NewObject<UInputBehaviorSet>(this, TEXT("GizmoInputBehaviors"));
+}
+
+void UInteractiveGizmo::Setup()
+{
+}
+
+void UInteractiveGizmo::Shutdown()
+{
+	InputBehaviors->RemoveAll();
+}
+
+void UInteractiveGizmo::Render(IToolsContextRenderAPI* RenderAPI)
+{
+}
+
+
+void UInteractiveGizmo::AddInputBehavior(UInputBehavior* Behavior)
+{
+	InputBehaviors->Add(Behavior);
+}
+
+const UInputBehaviorSet* UInteractiveGizmo::GetInputBehaviors() const
+{
+	return InputBehaviors;
+}
+
+
+
+void UInteractiveGizmo::Tick(float DeltaTime)
+{
+}
+
+
+UInteractiveGizmoManager* UInteractiveGizmo::GetGizmoManager() const
+{
+	UInteractiveGizmoManager* GizmoManager = Cast<UInteractiveGizmoManager>(GetOuter());
+	check(GizmoManager != nullptr);
+	return GizmoManager;
+}
\ No newline at end of file
diff --git a/Engine/Source/Runtime/InteractiveToolsFramework/Private/InteractiveGizmoManager.cpp b/Engine/Source/Runtime/InteractiveToolsFramework/Private/InteractiveGizmoManager.cpp
new file mode 100644
--- /dev/null
+++ b/Engine/Source/Runtime/InteractiveToolsFramework/Private/InteractiveGizmoManager.cpp
@@ -0,0 +1,228 @@
+// Copyright 1998-2019 Epic Games, Inc. All Rights Reserved.
+
+
+#include "InteractiveGizmoManager.h"
+
+
+UInteractiveGizmoManager::UInteractiveGizmoManager()
+{
+	QueriesAPI = nullptr;
+	TransactionsAPI = nullptr;
+	InputRouter = nullptr;
+}
+
+
+void UInteractiveGizmoManager::Initialize(IToolsContextQueriesAPI* QueriesAPIIn, IToolsContextTransactionsAPI* TransactionsAPIIn, UInputRouter* InputRouterIn)
+{
+	this->QueriesAPI = QueriesAPIIn;
+	this->TransactionsAPI = TransactionsAPIIn;
+	this->InputRouter = InputRouterIn;
+}
+
+
+void UInteractiveGizmoManager::Shutdown()
+{
+	this->QueriesAPI = nullptr;
+
+	TArray<FActiveGizmo> AllGizmos = ActiveGizmos;
+	for (FActiveGizmo& ActiveGizmo : AllGizmos)
+	{
+		DestroyGizmo(ActiveGizmo.Gizmo);
+	}
+	ActiveGizmos.Reset();
+
+	this->TransactionsAPI = nullptr;
+}
+
+
+
+void UInteractiveGizmoManager::RegisterGizmoType(const FString& Identifier, UInteractiveGizmoBuilder* Builder)
+{
+	check(GizmoBuilders.Contains(Identifier) == false);
+	GizmoBuilders.Add(Identifier, Builder );
+}
+
+
+bool UInteractiveGizmoManager::DeregisterGizmoType(const FString& BuilderIdentifier)
+{
+	if (GizmoBuilders.Contains(BuilderIdentifier) == false)
+	{
+		PostMessage(FString::Printf(TEXT("UInteractiveGizmoManager::DeregisterGizmoType: could not find requested type %s"), *BuilderIdentifier), EToolMessageLevel::Internal);
+		return false;
+	}
+	GizmoBuilders.Remove(BuilderIdentifier);
+	return true;
+}
+
+
+
+
+UInteractiveGizmo* UInteractiveGizmoManager::CreateGizmo(const FString& BuilderIdentifier, const FString& InstanceIdentifier)
+{
+	if ( GizmoBuilders.Contains(BuilderIdentifier) == false )
+	{
+		PostMessage(FString::Printf(TEXT("UInteractiveGizmoManager::CreateGizmo: could not find requested type %s"), *BuilderIdentifier), EToolMessageLevel::Internal);
+		return nullptr;
+	}
+	UInteractiveGizmoBuilder* FoundBuilder = GizmoBuilders[BuilderIdentifier];
+
+	// check if we have used this instance identifier
+	for (FActiveGizmo& ActiveGizmo : ActiveGizmos)
+	{
+		if (ActiveGizmo.InstanceIdentifier == InstanceIdentifier)
+		{
+			PostMessage(FString::Printf(TEXT("UInteractiveGizmoManager::CreateGizmo: instance identifier %s already in use!"), *InstanceIdentifier), EToolMessageLevel::Internal);
+			return nullptr;
+		}
+	}
+
+	FToolBuilderState CurrentSceneState;
+	QueriesAPI->GetCurrentSelectionState(CurrentSceneState);
+
+	UInteractiveGizmo* NewGizmo = FoundBuilder->BuildGizmo(CurrentSceneState);
+	if (NewGizmo == nullptr)
+	{
+		PostMessage(FString::Printf(TEXT("UInteractiveGizmoManager::CreateGizmo: BuildGizmo() returned null")), EToolMessageLevel::Internal);
+		return false;
+	}
+
+	NewGizmo->Setup();
+
+	// register new active input behaviors
+	InputRouter->RegisterSource(NewGizmo);
+
+	PostInvalidation();
+
+	FActiveGizmo ActiveGizmo = { NewGizmo, BuilderIdentifier, InstanceIdentifier };
+	ActiveGizmos.Add(ActiveGizmo);
+
+	return NewGizmo;
+}
+
+
+
+bool UInteractiveGizmoManager::DestroyGizmo(UInteractiveGizmo* Gizmo)
+{
+	int FoundIndex = -1;
+	for ( int i = 0; i < ActiveGizmos.Num(); ++i )
+	{
+		if (ActiveGizmos[i].Gizmo == Gizmo)
+		{
+			FoundIndex = i;
+			break;
+		}
+	}
+	if (FoundIndex == -1)
+	{
+		return false;
+	}
+
+	InputRouter->ForceTerminateSource(Gizmo);
+
+	Gizmo->Shutdown();
+
+	InputRouter->DeregisterSource(Gizmo);
+
+	ActiveGizmos.RemoveAt(FoundIndex);
+
+	PostInvalidation();
+
+	return true;
+}
+
+
+
+
+TArray<UInteractiveGizmo*> UInteractiveGizmoManager::FindAllGizmosOfType(const FString& BuilderIdentifier)
+{
+	TArray<UInteractiveGizmo*> Found;
+	for (int i = 0; i < ActiveGizmos.Num(); ++i)
+	{
+		if (ActiveGizmos[i].BuilderIdentifier == BuilderIdentifier)
+		{
+			Found.Add(ActiveGizmos[i].Gizmo);
+		}
+	}
+	return Found;
+}
+
+
+void UInteractiveGizmoManager::DestroyAllGizmosOfType(const FString& BuilderIdentifier)
+{
+	TArray<UInteractiveGizmo*> ToRemove = FindAllGizmosOfType(BuilderIdentifier);
+
+	for (int i = 0; i < ToRemove.Num(); ++i)
+	{
+		DestroyGizmo(ToRemove[i]);
+	}
+}
+
+
+
+UInteractiveGizmo* UInteractiveGizmoManager::FindGizmoByInstanceIdentifier(const FString& Identifier)
+{
+	for (int i = 0; i < ActiveGizmos.Num(); ++i)
+	{
+		if (ActiveGizmos[i].InstanceIdentifier == Identifier)
+		{
+			return ActiveGizmos[i].Gizmo;
+		}
+	}
+	return nullptr;
+}
+
+
+
+void UInteractiveGizmoManager::Tick(float DeltaTime)
+{
+	for (FActiveGizmo& ActiveGizmo : ActiveGizmos)
+	{
+		ActiveGizmo.Gizmo->Tick(DeltaTime);
+	}
+}
+
+
+void UInteractiveGizmoManager::Render(IToolsContextRenderAPI* RenderAPI)
+{
+	for (FActiveGizmo& ActiveGizmo : ActiveGizmos)
+	{
+		ActiveGizmo.Gizmo->Render(RenderAPI);
+	}
+
+}
+
+
+void UInteractiveGizmoManager::PostMessage(const TCHAR* Message, EToolMessageLevel Level)
+{
+	TransactionsAPI->PostMessage(Message, Level);
+}
+
+void UInteractiveGizmoManager::PostMessage(const FString& Message, EToolMessageLevel Level)
+{
+	TransactionsAPI->PostMessage(*Message, Level);
+}
+
+void UInteractiveGizmoManager::PostInvalidation()
+{
+	TransactionsAPI->PostInvalidation();
+}
+
+
+void UInteractiveGizmoManager::BeginUndoTransaction(const FText& Description)
+{
+	TransactionsAPI->BeginUndoTransaction(Description);
+}
+
+void UInteractiveGizmoManager::EndUndoTransaction()
+{
+	TransactionsAPI->EndUndoTransaction();
+}
+
+
+
+void UInteractiveGizmoManager::EmitObjectChange(UObject* TargetObject, TUniquePtr<FChange> Change, const FText& Description)
+{
+	TransactionsAPI->AppendChange(TargetObject, MoveTemp(Change), Description );
+}
+
+
diff --git a/Engine/Source/Runtime/InteractiveToolsFramework/Private/InteractiveToolsContext.cpp b/Engine/Source/Runtime/InteractiveToolsFramework/Private/InteractiveToolsContext.cpp
--- a/Engine/Source/Runtime/InteractiveToolsFramework/Private/InteractiveToolsContext.cpp
+++ b/Engine/Source/Runtime/InteractiveToolsFramework/Private/InteractiveToolsContext.cpp
@@ -17,6 +17,9 @@ void UInteractiveToolsContext::Initialize(IToolsContextQueriesAPI* QueriesAPI, I
 
 	ToolManager = NewObject<UInteractiveToolManager>(this);
 	ToolManager->Initialize(QueriesAPI, TransactionsAPI, InputRouter);
+
+	GizmoManager = NewObject<UInteractiveGizmoManager>(this);
+	GizmoManager->Initialize(QueriesAPI, TransactionsAPI, InputRouter);
 }
 
 
@@ -27,6 +30,9 @@ void UInteractiveToolsContext::Shutdown()
 	InputRouter->Shutdown();
 	InputRouter = nullptr;
 
+	GizmoManager->Shutdown();
+	GizmoManager = nullptr;
+
 	ToolManager->Shutdown();
 	ToolManager = nullptr;
 }
diff --git a/Engine/Source/Runtime/InteractiveToolsFramework/Public/InteractiveGizmo.h b/Engine/Source/Runtime/InteractiveToolsFramework/Public/InteractiveGizmo.h
new file mode 100644
--- /dev/null
+++ b/Engine/Source/Runtime/InteractiveToolsFramework/Public/InteractiveGizmo.h
@@ -0,0 +1,83 @@
+// Copyright 1998-2019 Epic Games, Inc. All Rights Reserved.
+
+#pragma once
+
+#include "CoreMinimal.h"
+#include "InputBehaviorSet.h"
+#include "ToolContextInterfaces.h"
+#include "InteractiveGizmo.generated.h"
+
+class UInteractiveGizmoManager;
+
+
+
+/**
+ * UInteractiveGizmo is the base class for all Gizmos in the InteractiveToolsFramework.
+ *
+ * @todo callback/delegate for if/when .InputBehaviors changes
+ * @todo callback/delegate for when Gizmo properties change
+ */
+UCLASS(Transient)
+class INTERACTIVETOOLSFRAMEWORK_API UInteractiveGizmo : public UObject, public IInputBehaviorSource
+{
+	GENERATED_BODY()
+
+public:
+	UInteractiveGizmo();
+
+	/**
+	 * Called by GizmoManager to initialize the Gizmo *after* GizmoBuilder::BuildGizmo() has been called
+	 */
+	virtual void Setup();
+
+	/**
+	 * Called by GizmoManager to shut down the Gizmo
+	 */
+	virtual void Shutdown();
+
+	/**
+	 * Allow the Gizmo to do any custom drawing (ie via PDI/RHI)
+	 * @param RenderAPI Abstraction that provides access to Rendering in the current ToolsContext
+	 */
+	virtual void Render(IToolsContextRenderAPI* RenderAPI);
+
+	/**
+	 * Allow the Gizmo to do any necessary processing on Tick 
+	 * @param DeltaTime the time delta since last tick
+	 */
+	virtual void Tick(float DeltaTime);
+
+
+
+	/**
+	 * @return GizmoManager that owns this Gizmo
+	 */
+	virtual UInteractiveGizmoManager* GetGizmoManager() const;
+
+
+
+	//
+	// Input Behaviors support
+	//
+
+	/**
+	 * Add an input behavior for this Gizmo
+	 * @param Behavior behavior to add
+	 */
+	virtual void AddInputBehavior(UInputBehavior* Behavior);
+
+	/**
+	 * @return Current input behavior set.
+	 */
+	virtual const UInputBehaviorSet* GetInputBehaviors() const;
+
+
+
+protected:
+
+	/** The current set of InputBehaviors provided by this Gizmo */
+	UPROPERTY()
+	UInputBehaviorSet* InputBehaviors;
+};
+
+
diff --git a/Engine/Source/Runtime/InteractiveToolsFramework/Public/InteractiveGizmoBuilder.h b/Engine/Source/Runtime/InteractiveToolsFramework/Public/InteractiveGizmoBuilder.h
new file mode 100644
--- /dev/null
+++ b/Engine/Source/Runtime/InteractiveToolsFramework/Public/InteractiveGizmoBuilder.h
@@ -0,0 +1,34 @@
+// Copyright 1998-2019 Epic Games, Inc. All Rights Reserved.
+
+#pragma once
+
+#include "CoreMinimal.h"
+#include "UObject/Interface.h"
+#include "InteractiveGizmo.h"
+#include "InteractiveGizmoBuilder.generated.h"
+
+
+
+/**
+ * A UInteractiveGizmoBuilder creates a new instance of an InteractiveGizmo (basically this is a Factory).
+ * These are registered with the InteractiveGizmoManager, which calls BuildGizmo().
+ * This is an abstract base class, you must subclass it in order to create your particular Gizmo instance
+ */
+UCLASS(Transient, Abstract)
+class INTERACTIVETOOLSFRAMEWORK_API UInteractiveGizmoBuilder : public UObject
+{
+	GENERATED_BODY()
+
+public:
+
+	/** 
+	 * Create a new instance of this builder's Gizmo
+	 * @param SceneState the current scene selection state, etc
+	 * @return a new instance of the Gizmo, or nullptr on error/failure
+	 */
+	virtual UInteractiveGizmo* BuildGizmo(const FToolBuilderState& SceneState) const
+	{
+		unimplemented();
+		return nullptr;
+	}
+};
\ No newline at end of file
diff --git a/Engine/Source/Runtime/InteractiveToolsFramework/Public/InteractiveGizmoManager.h b/Engine/Source/Runtime/InteractiveToolsFramework/Public/InteractiveGizmoManager.h
new file mode 100644
--- /dev/null
+++ b/Engine/Source/Runtime/InteractiveToolsFramework/Public/InteractiveGizmoManager.h
@@ -0,0 +1,170 @@
+// Copyright 1998-2019 Epic Games, Inc. All Rights Reserved.
+
+#pragma once
+
+#include "CoreMinimal.h"
+#include "UObject/Object.h"
+#include "Misc/Change.h"
+#include "InteractiveGizmo.h"
+#include "InteractiveGizmoBuilder.h"
+#include "InputRouter.h"
+#include "ToolContextInterfaces.h"
+#include "InteractiveGizmoManager.generated.h"
+
+
+
+USTRUCT()
+struct FActiveGizmo
+{
+	GENERATED_BODY();
+
+	UInteractiveGizmo* Gizmo;
+	FString BuilderIdentifier;
+	FString InstanceIdentifier;
+};
+
+
+/**
+ * UInteractiveGizmoManager allows users of the Tools framework to create and operate Gizmo instances.
+ * For each Gizmo, a (string,GizmoBuilder) pair is registered with the GizmoManager.
+ * Gizmos can then be activated via the string identifier.
+ * 
+ */
+UCLASS(Transient)
+class INTERACTIVETOOLSFRAMEWORK_API UInteractiveGizmoManager : public UObject
+{
+	GENERATED_BODY()
+
+protected:
+	friend class UInteractiveToolsContext;		// to call Initialize/Shutdown
+
+	UInteractiveGizmoManager();
+
+	/** Initialize the GizmoManager with the necessary Context-level state. UInteractiveToolsContext calls this, you should not. */
+	virtual void Initialize(IToolsContextQueriesAPI* QueriesAPI, IToolsContextTransactionsAPI* TransactionsAPI, UInputRouter* InputRouter);
+
+	/** Shutdown the GizmoManager. Called by UInteractiveToolsContext. */
+	virtual void Shutdown();
+
+public:
+
+	//
+	// GizmoBuilder Registration and Gizmo Creation/Shutdown
+	//
+
+	/**
+	 * Register a new GizmoBuilder
+	 * @param BuilderIdentifier string used to identify this Builder
+	 * @param Builder new GizmoBuilder instance
+	 */
+	virtual void RegisterGizmoType(const FString& BuilderIdentifier, UInteractiveGizmoBuilder* Builder);
+
+	/**
+	 * Remove a GizmoBuilder from the set of known GizmoBuilders
+	 * @param BuilderIdentifier identification string that was passed to RegisterGizmoType()
+	 * @return true if Builder was found and deregistered
+	 */
+	virtual bool DeregisterGizmoType(const FString& BuilderIdentifier);
+
+
+
+	/**
+	 * Try to activate a new Gizmo instance on the given Side
+	 * @param BuilderIdentifier string used to identify Builder that should be called
+	 * @param InstanceIdentifier client-defined string that can be used to locate this instance
+	 * @return new Gizmo instance that has been created and initialized
+	 */	
+	virtual UInteractiveGizmo* CreateGizmo(const FString& BuilderIdentifier, const FString& InstanceIdentifier );
+
+
+	/**
+	 * Shutdown and remove a Gizmo
+	 * @param Gizmo the Gizmo to shutdown and remove
+	 * @return true if the Gizmo was found and removed
+	 */
+	virtual bool DestroyGizmo(UInteractiveGizmo* Gizmo);
+
+	/**
+	 * Destroy all Gizmos that were created by the identified GizmoBuilder
+	 * @param BuilderIdentifier the Builder string registered with RegisterGizmoType
+	 */
+	virtual void DestroyAllGizmosOfType(const FString& BuilderIdentifier);
+
+	/**
+	 * Find all the existing Gizmo instances that were created by the identified GizmoBuilder
+	 * @param BuilderIdentifier the Builder string registered with RegisterGizmoType
+	 * @return list of found Gizmos
+	 */
+	virtual TArray<UInteractiveGizmo*> FindAllGizmosOfType(const FString& BuilderIdentifier);
+
+	/**
+	 * Find the Gizmo that was created with the given instance identifier
+	 * @param Identifier the InstanceIdentifier that was passed to CreateGizmo()
+	 * @return the found Gizmo, or null
+	 */
+	virtual UInteractiveGizmo* FindGizmoByInstanceIdentifier(const FString& Identifier);
+
+
+
+
+	//
+	// Functions that Gizmos can call to interact with Transactions API
+	//
+	
+	/** Post a message via the Transactions API */
+	virtual void PostMessage(const TCHAR* Message, EToolMessageLevel Level);
+
+	/** Post a message via the Transactions API */
+	virtual void PostMessage(const FString& Message, EToolMessageLevel Level);
+
+	/** Request an Invalidation via the Transactions API (ie to cause a repaint, etc) */
+	virtual void PostInvalidation();
+
+	/**
+	 * Request that the Context open a Transaction, whatever that means to the current Context
+	 * @param Description text description of this transaction (this is the string that appears on undo/redo in the UE Editor)
+	 */
+	virtual void BeginUndoTransaction(const FText& Description);
+
+	/** Request that the Context close and commit the open Transaction */
+	virtual void EndUndoTransaction();
+
+	/**
+	 * Forward an FChange object to the Context
+	 * @param TargetObject the object that the FChange applies to
+	 * @param Change the change object that the Context should insert into the transaction history
+	 * @param Description text description of this change (this is the string that appears on undo/redo in the UE Editor)
+	 */
+	virtual void EmitObjectChange(UObject* TargetObject, TUniquePtr<FChange> Change, const FText& Description );
+
+
+	//
+	// State control  (@todo: have the Context call these? not safe for anyone to call)
+	//
+
+	/** Tick any active Gizmos. Called by UInteractiveToolsContext */
+	virtual void Tick(float DeltaTime);
+
+	/** Render any active Gizmos. Called by UInteractiveToolsContext. */
+	virtual void Render(IToolsContextRenderAPI* RenderAPI);
+
+
+public:
+	/** set of Currently-active Gizmos */
+	UPROPERTY()
+	TArray<FActiveGizmo> ActiveGizmos;
+
+
+protected:
+	/** Current Context-Queries implementation */
+	IToolsContextQueriesAPI* QueriesAPI;
+	/** Current Transactions implementation */
+	IToolsContextTransactionsAPI* TransactionsAPI;
+
+	/** Current InputRouter (Context owns this) */
+	UInputRouter* InputRouter;
+
+	/** Current set of named GizmoBuilders */
+	UPROPERTY()
+	TMap<FString, UInteractiveGizmoBuilder*> GizmoBuilders;
+};
\ No newline at end of file
diff --git a/Engine/Source/Runtime/InteractiveToolsFramework/Public/InteractiveToolsContext.h b/Engine/Source/Runtime/InteractiveToolsFramework/Public/InteractiveToolsContext.h
--- a/Engine/Source/Runtime/InteractiveToolsFramework/Public/InteractiveToolsContext.h
+++ b/Engine/Source/Runtime/InteractiveToolsFramework/Public/InteractiveToolsContext.h
@@ -4,6 +4,7 @@
 
 #include "CoreMinimal.h"
 #include "InteractiveToolManager.h"
+#include "InteractiveGizmoManager.h"
 #include "InteractiveToolsContext.generated.h"
 
 /**
@@ -38,4 +39,7 @@ public:
 	UPROPERTY()
 	UInteractiveToolManager* ToolManager;	
 
+	/** current UInteractiveGizmoManager for this Context */
+	UPROPERTY()
+	UInteractiveGizmoManager* GizmoManager;
 };
\ No newline at end of file
diff --git a/Engine/Source/Runtime/InteractiveToolsFramework/Public/ToolContextInterfaces.h b/Engine/Source/Runtime/InteractiveToolsFramework/Public/ToolContextInterfaces.h
--- a/Engine/Source/Runtime/InteractiveToolsFramework/Public/ToolContextInterfaces.h
+++ b/Engine/Source/Runtime/InteractiveToolsFramework/Public/ToolContextInterfaces.h
@@ -12,7 +12,7 @@ class UPackage;
 class FPrimitiveDrawInterface;
 class USelection;
 class UInteractiveToolManager;
-
+class UInteractiveGizmoManager;
 
 /**
  * FToolBuilderState is a bucket of state information that a ToolBuilder might need
@@ -22,25 +22,22 @@ class UInteractiveToolManager;
 struct INTERACTIVETOOLSFRAMEWORK_API FToolBuilderState
 {
 	/** The current UWorld */
-	UWorld* World;
+	UWorld* World = nullptr;
 	/** The current ToolManager */
-	UInteractiveToolManager* ToolManager;
+	UInteractiveToolManager* ToolManager = nullptr;
+	/** The current GizmoManager */
+	UInteractiveGizmoManager* GizmoManager = nullptr;
 
 	/** Current selected Actors. May be empty or nullptr. */
-	USelection* SelectedActors;
+	USelection* SelectedActors = nullptr;
 	/** Current selected Components. May be empty or nullptr. */
-	USelection* SelectedComponents;
+	USelection* SelectedComponents = nullptr;
 
 	/** Implementation that can build Sources (like MeshDescriptionSource) for Components */
-	IComponentSourceFactory* SourceBuilder;
+	IComponentSourceFactory* SourceBuilder = nullptr;
 
 	FToolBuilderState()
 	{
-		World = nullptr;
-		ToolManager = nullptr;
-		SelectedActors = nullptr;
-		SelectedComponents = nullptr;
-		SourceBuilder = nullptr;
 	}
 };
 
