build fixes
#rnx

diff --git a/Engine/Source/Runtime/InteractiveToolsFramework/Private/InteractiveGizmoManager.cpp b/Engine/Source/Runtime/InteractiveToolsFramework/Private/InteractiveGizmoManager.cpp
--- a/Engine/Source/Runtime/InteractiveToolsFramework/Private/InteractiveGizmoManager.cpp
+++ b/Engine/Source/Runtime/InteractiveToolsFramework/Private/InteractiveGizmoManager.cpp
@@ -83,7 +83,7 @@ UInteractiveGizmo* UInteractiveGizmoManager::CreateGizmo(const FString& BuilderI
 	if (NewGizmo == nullptr)
 	{
 		PostMessage(FString::Printf(TEXT("UInteractiveGizmoManager::CreateGizmo: BuildGizmo() returned null")), EToolMessageLevel::Internal);
-		return false;
+		return nullptr;
 	}
 
 	NewGizmo->Setup();
