minor tools improvements - brush indicator line points out, mesh vertex drawing functions
#rnx

diff --git a/Engine/Source/Runtime/InteractiveToolsFramework/Private/BaseTools/MeshSurfacePointTool.cpp b/Engine/Source/Runtime/InteractiveToolsFramework/Private/BaseTools/MeshSurfacePointTool.cpp
--- a/Engine/Source/Runtime/InteractiveToolsFramework/Private/BaseTools/MeshSurfacePointTool.cpp
+++ b/Engine/Source/Runtime/InteractiveToolsFramework/Private/BaseTools/MeshSurfacePointTool.cpp
@@ -31,7 +31,7 @@ void UMeshSurfacePointToolBuilder::InitializeNewTool(UMeshSurfacePointTool* NewT
 {
 	UObject* MeshComponent = ToolBuilderUtil::FindFirstComponent(SceneState, ToolBuilderUtil::IsMeshDescriptionSourceComponent);
 	check(MeshComponent != nullptr);
-	NewTool->SetSelection(
+	NewTool->SetMeshSource(
 		SceneState.SourceBuilder->MakeMeshDescriptionSource(MeshComponent) );
 }
 
@@ -40,7 +40,7 @@ void UMeshSurfacePointToolBuilder::InitializeNewTool(UMeshSurfacePointTool* NewT
  * Tool
  */
 
-void UMeshSurfacePointTool::SetSelection(TUniquePtr<IMeshDescriptionSource> MeshSourceIn)
+void UMeshSurfacePointTool::SetMeshSource(TUniquePtr<IMeshDescriptionSource> MeshSourceIn)
 {
 	this->MeshSource = TUniquePtr<IMeshDescriptionSource>(MoveTemp(MeshSourceIn));
 }
diff --git a/Engine/Source/Runtime/InteractiveToolsFramework/Public/BaseTools/MeshSurfacePointTool.h b/Engine/Source/Runtime/InteractiveToolsFramework/Public/BaseTools/MeshSurfacePointTool.h
--- a/Engine/Source/Runtime/InteractiveToolsFramework/Public/BaseTools/MeshSurfacePointTool.h
+++ b/Engine/Source/Runtime/InteractiveToolsFramework/Public/BaseTools/MeshSurfacePointTool.h
@@ -66,7 +66,7 @@ public:
 	 *        only calls MeshSource->HitTest, so this source technically does not have to actually provide
 	 *        a valid MeshDescription
 	 */
-	virtual void SetSelection(TUniquePtr<IMeshDescriptionSource> MeshSource);
+	virtual void SetMeshSource(TUniquePtr<IMeshDescriptionSource> MeshSource);
 
 	/**
 	 * @return true if the target MeshSource is hit by the Ray
