MouseHoverBehavior now supports modifier keys
#rnx

diff --git a/Engine/Source/Runtime/InteractiveToolsFramework/Private/BaseBehaviors/MouseHoverBehavior.cpp b/Engine/Source/Runtime/InteractiveToolsFramework/Private/BaseBehaviors/MouseHoverBehavior.cpp
--- a/Engine/Source/Runtime/InteractiveToolsFramework/Private/BaseBehaviors/MouseHoverBehavior.cpp
+++ b/Engine/Source/Runtime/InteractiveToolsFramework/Private/BaseBehaviors/MouseHoverBehavior.cpp
@@ -27,6 +27,8 @@ void UMouseHoverBehavior::UpdateHover(const FInputDeviceState& input)
 {
 	if (Target != nullptr)
 	{
+		Modifiers.UpdateModifiers(input, Target);
+
 		Target->OnUpdateHover( FInputDeviceRay(input.Mouse.WorldRay, input.Mouse.Position2D) );
 	}
 }
diff --git a/Engine/Source/Runtime/InteractiveToolsFramework/Public/BaseBehaviors/BehaviorTargetInterfaces.h b/Engine/Source/Runtime/InteractiveToolsFramework/Public/BaseBehaviors/BehaviorTargetInterfaces.h
--- a/Engine/Source/Runtime/InteractiveToolsFramework/Public/BaseBehaviors/BehaviorTargetInterfaces.h
+++ b/Engine/Source/Runtime/InteractiveToolsFramework/Public/BaseBehaviors/BehaviorTargetInterfaces.h
@@ -150,7 +150,7 @@ public:
  * device event data in a generic way, without requiring that all Tools
  * know about the concept of Hovering.
  */
-class IHoverBehaviorTarget
+class IHoverBehaviorTarget : public IModifierToggleBehaviorTarget
 {
 public:
 	virtual ~IHoverBehaviorTarget() {}
diff --git a/Engine/Source/Runtime/InteractiveToolsFramework/Public/BaseBehaviors/MouseHoverBehavior.h b/Engine/Source/Runtime/InteractiveToolsFramework/Public/BaseBehaviors/MouseHoverBehavior.h
--- a/Engine/Source/Runtime/InteractiveToolsFramework/Public/BaseBehaviors/MouseHoverBehavior.h
+++ b/Engine/Source/Runtime/InteractiveToolsFramework/Public/BaseBehaviors/MouseHoverBehavior.h
@@ -6,6 +6,7 @@
 #include "InputBehavior.h"
 #include "InteractiveTool.h"
 #include "BaseBehaviors/BehaviorTargetInterfaces.h"
+#include "BaseBehaviors/InputBehaviorModifierStates.h"
 #include "MouseHoverBehavior.generated.h"
 
 
@@ -22,6 +23,11 @@ class INTERACTIVETOOLSFRAMEWORK_API UMouseHoverBehavior : public UInputBehavior
 public:
 	UMouseHoverBehavior();
 
+	/**
+	 * The modifier set for this behavior
+	 */
+	FInputBehaviorModifierStates Modifiers;
+
 	virtual void Initialize(IHoverBehaviorTarget* Target);
 
 	// UInputBehavior hover implementation
