Add InteractiveToolActionSet to InteractiveTool
#rnx

diff --git a/Engine/Source/Runtime/InteractiveToolsFramework/Private/InteractiveTool.cpp b/Engine/Source/Runtime/InteractiveToolsFramework/Private/InteractiveTool.cpp
--- a/Engine/Source/Runtime/InteractiveToolsFramework/Private/InteractiveTool.cpp
+++ b/Engine/Source/Runtime/InteractiveToolsFramework/Private/InteractiveTool.cpp
@@ -66,6 +66,27 @@ const TArray<UObject*>& UInteractiveTool::GetToolProperties() const
 }
 
 
+void UInteractiveTool::RegisterActions(FInteractiveToolActionSet& ActionSet)
+{
+}
+
+FInteractiveToolActionSet* UInteractiveTool::GetActionSet()
+{
+	if (ToolActionSet == nullptr)
+	{
+		ToolActionSet = new FInteractiveToolActionSet();
+		RegisterActions(*ToolActionSet);
+	}
+	return ToolActionSet;
+}
+
+void UInteractiveTool::ExecuteAction(int32 ActionID)
+{
+	GetActionSet()->ExecuteAction(ActionID);
+}
+
+
+
 bool UInteractiveTool::HasCancel() const
 {
 	return false;
@@ -86,7 +107,6 @@ void UInteractiveTool::Tick(float DeltaTime)
 {
 }
 
-
 UInteractiveToolManager* UInteractiveTool::GetToolManager() const
 {
 	UInteractiveToolManager* ToolManager = Cast<UInteractiveToolManager>(GetOuter());
diff --git a/Engine/Source/Runtime/InteractiveToolsFramework/Private/InteractiveToolActionSet.cpp b/Engine/Source/Runtime/InteractiveToolsFramework/Private/InteractiveToolActionSet.cpp
new file mode 100644
--- /dev/null
+++ b/Engine/Source/Runtime/InteractiveToolsFramework/Private/InteractiveToolActionSet.cpp
@@ -0,0 +1,59 @@
+// Copyright 1998-2019 Epic Games, Inc. All Rights Reserved.
+
+#include "InteractiveToolActionSet.h"
+#include "InteractiveTool.h"
+
+
+
+void FInteractiveToolActionSet::RegisterAction(UInteractiveTool* Tool, int32 ActionID,
+	const FString& ActionName, const FString& ShortUIName, const FString& DescriptionText,
+	EModifierKey::Type Modifiers, const FKey& ShortcutKey,
+	TFunction<void()> ActionFunction )
+{
+	checkf(FindActionByID(ActionID) == nullptr, TEXT("InteractiveToolActionSet::RegisterAction: Action ID is already registered!"));
+
+	FInteractiveToolAction NewAction;
+	NewAction.ClassType = Tool->GetClass();
+	NewAction.ActionID = ActionID;
+	NewAction.ActionName = ActionName;
+	NewAction.ShortName = ShortUIName;
+	NewAction.Description = DescriptionText;
+	NewAction.DefaultModifiers = Modifiers;
+	NewAction.DefaultKey = ShortcutKey;
+	NewAction.OnAction = ActionFunction;
+
+	Actions.Add(NewAction);
+}
+
+
+const FInteractiveToolAction* FInteractiveToolActionSet::FindActionByID(int32 ActionID) const
+{
+	for (const FInteractiveToolAction& Action : Actions)
+	{
+		if (Action.ActionID == ActionID)
+		{
+			return &Action;
+		}
+	}
+	return nullptr;
+}
+
+
+
+void FInteractiveToolActionSet::CollectActions(TArray<FInteractiveToolAction>& OutActions) const
+{
+	for (const FInteractiveToolAction& Action : Actions)
+	{
+		OutActions.Add(Action);
+	}
+}
+
+
+void FInteractiveToolActionSet::ExecuteAction(int32 ActionID) const
+{
+	const FInteractiveToolAction* Found = FindActionByID(ActionID);
+	if (Found != nullptr)
+	{
+		Found->OnAction();
+	}
+}
\ No newline at end of file
diff --git a/Engine/Source/Runtime/InteractiveToolsFramework/Public/InteractiveTool.h b/Engine/Source/Runtime/InteractiveToolsFramework/Public/InteractiveTool.h
--- a/Engine/Source/Runtime/InteractiveToolsFramework/Public/InteractiveTool.h
+++ b/Engine/Source/Runtime/InteractiveToolsFramework/Public/InteractiveTool.h
@@ -4,6 +4,7 @@
 
 #include "CoreMinimal.h"
 #include "InputBehaviorSet.h"
+#include "InteractiveToolActionSet.h"
 #include "ToolContextInterfaces.h"
 #include "InteractiveTool.generated.h"
 
@@ -178,6 +179,7 @@ protected:
 	UInputBehaviorSet* InputBehaviors;
 
 	/** The current set of Property UObjects provided by this Tool. May contain pointer to itself. */
+	UPROPERTY()
 	TArray<UObject*> ToolPropertyObjects;
 
 	/**
@@ -191,6 +193,47 @@ protected:
 	 * @param PropertySet Property Set object to add
 	 */
 	virtual void AddToolPropertySource(UInteractiveToolPropertySet* PropertySet);
+
+
+
+
+	//
+	// Action support/system
+	// 
+	// Your Tool subclass can register a set of "Actions" it can execute
+	// by overloading RegisterActions(). Then external systems can use GetActionSet() to
+	// find out what Actions your Tool supports, and ExecuteAction() to run those actions.
+	//
+	
+public:
+	/**
+	 * Get the internal Action Set for this Tool. The action set is created and registered on-demand.
+	 * @return pointer to initialized Action set
+	 */
+	virtual FInteractiveToolActionSet* GetActionSet();
+
+	/**
+	 * Request that the Action identified by ActionID be executed.
+	 * Default implementation forwards these requests to internal ToolActionSet.
+	 */
+	virtual void ExecuteAction(int32 ActionID);
+
+
+protected:
+	/**
+	 * Override this function to register the set of Actions this Tool supports, using FInteractiveToolActionSet::RegisterAction.
+	 * Note that 
+	 */
+	virtual void RegisterActions(FInteractiveToolActionSet& ActionSet);
+
+
+private:
+	/** 
+	 * Set of actions this Tool can execute. This variable is allocated on-demand. 
+	 * Use GetActionSet() instead of accessing this pointer directly!
+	 */
+	FInteractiveToolActionSet* ToolActionSet = nullptr;
+
 };
 
 
diff --git a/Engine/Source/Runtime/InteractiveToolsFramework/Public/InteractiveToolActionSet.h b/Engine/Source/Runtime/InteractiveToolsFramework/Public/InteractiveToolActionSet.h
new file mode 100644
--- /dev/null
+++ b/Engine/Source/Runtime/InteractiveToolsFramework/Public/InteractiveToolActionSet.h
@@ -0,0 +1,113 @@
+// Copyright 1998-2019 Epic Games, Inc. All Rights Reserved.
+
+#pragma once
+
+#include "CoreMinimal.h"
+
+class UInteractiveTool;
+
+
+
+/**
+ * Standard Actions that can be shared across multiple Tools.
+ * The enum values are used as FInteractiveToolAction.ActionID.
+ * If you need to define your own values, start at BaseClientDefinedActionID
+ */
+enum class EStandardToolActions
+{
+	IncreaseBrushSize = 100,
+	DecreaseBrushSize = 101,
+
+
+	BaseClientDefinedActionID = 10000
+};
+
+
+
+/**
+ * FInteractiveToolAction is returned by a UInteractiveTool to represent
+ * an "Action" the Tool can execute. 
+ */
+struct INTERACTIVETOOLSFRAMEWORK_API FInteractiveToolAction
+{
+	/** Which type of UInteractiveTool this Action can be applied to*/
+	const UClass* ClassType;
+
+	/** Identifier for this Action */
+	int32 ActionID;
+
+	/** Internal name for this Action */
+	FString ActionName;
+	/** Short name for this Action */
+	FString ShortName;
+	/** Descriptive name for this Action */
+	FString Description;
+
+	/** Suggested modifier keys for this Action */
+	EModifierKey::Type DefaultModifiers;
+	/** Suggested keybinding for this Action. */
+	FKey DefaultKey;
+
+	/** Call this function to execute the Action */
+	TFunction<void()> OnAction;
+
+
+	FInteractiveToolAction()
+	{
+		ClassType = nullptr;
+		ActionID = 0;
+	}
+
+	FInteractiveToolAction(const UClass* ClassTypeIn, int32 ActionIDIn,
+		const FString& ActionNameIn, const FString& ShortNameIn, const FString& DescriptionIn,
+		EModifierKey::Type DefaultModifiersIn, const FKey& DefaultKeyIn )
+	{
+		ClassType = ClassTypeIn;
+		ActionID = ActionIDIn;
+		ActionName = ActionNameIn;
+		ShortName = ShortNameIn;
+		Description = DescriptionIn;
+		DefaultModifiers = DefaultModifiersIn;
+		DefaultKey = DefaultKeyIn;
+	}
+
+};
+
+
+
+/**
+ * FInteractiveToolActionSet maintains a list of FInteractiveToolAction.
+ * Each UInteractiveTool contains an instance of this class.
+ */
+class INTERACTIVETOOLSFRAMEWORK_API FInteractiveToolActionSet
+{
+public:
+
+	/**
+	 * Register an Action with the ActionSet. This function is intended to be called by
+	 * UInteractiveTool::RegisterActions() implementations
+	 */
+	void RegisterAction(UInteractiveTool* Tool, int32 ActionID,
+		const FString& ActionName, const FString& ShortUIName, const FString& DescriptionText,
+		EModifierKey::Type Modifiers, const FKey& ShortcutKey,
+		TFunction<void()> ActionFunction );
+
+	/**
+	 * Find an existing Action by ID
+	 * @return located Action, or nullptr if not found
+	 */
+	const FInteractiveToolAction* FindActionByID(int32 ActionID) const;
+
+	/**
+	 * Return the internal list of registered Actions by adding to the OutActions array
+	 */
+	void CollectActions(TArray<FInteractiveToolAction>& OutActions) const;
+
+	/**
+	 * Execute the action identified by ActionID
+	 */
+	void ExecuteAction(int32 ActionID) const;
+
+protected:
+	TArray<FInteractiveToolAction> Actions;
+};
\ No newline at end of file
