add GenericApplication.h header to get EModifierKey definition
#rnx

diff --git a/Engine/Source/Runtime/InteractiveToolsFramework/Public/InteractiveToolActionSet.h b/Engine/Source/Runtime/InteractiveToolsFramework/Public/InteractiveToolActionSet.h
--- a/Engine/Source/Runtime/InteractiveToolsFramework/Public/InteractiveToolActionSet.h
+++ b/Engine/Source/Runtime/InteractiveToolsFramework/Public/InteractiveToolActionSet.h
@@ -3,6 +3,7 @@
 #pragma once
 
 #include "CoreMinimal.h"
+#include "GenericPlatform/GenericApplication.h"
 
 class UInteractiveTool;
 
