add missing header to get FKey definition
#rnx

diff --git a/Engine/Source/Runtime/InteractiveToolsFramework/Public/InteractiveToolActionSet.h b/Engine/Source/Runtime/InteractiveToolsFramework/Public/InteractiveToolActionSet.h
--- a/Engine/Source/Runtime/InteractiveToolsFramework/Public/InteractiveToolActionSet.h
+++ b/Engine/Source/Runtime/InteractiveToolsFramework/Public/InteractiveToolActionSet.h
@@ -4,11 +4,11 @@
 
 #include "CoreMinimal.h"
 #include "GenericPlatform/GenericApplication.h"
+#include "InputCoreTypes.h"
 
 class UInteractiveTool;
 
 
-
 /**
  * Standard Actions that can be shared across multiple Tools.
  * The enum values are used as FInteractiveToolAction.ActionID.
