fix link error and linux compile error
#rnx

diff --git a/Engine/Source/Editor/EditorInteractiveToolsFramework/Public/InteractiveToolsCommands.h b/Engine/Source/Editor/EditorInteractiveToolsFramework/Public/InteractiveToolsCommands.h
--- a/Engine/Source/Editor/EditorInteractiveToolsFramework/Public/InteractiveToolsCommands.h
+++ b/Engine/Source/Editor/EditorInteractiveToolsFramework/Public/InteractiveToolsCommands.h
@@ -70,7 +70,7 @@ protected:
 
 	/** Forwarding constructor */
 	TInteractiveToolCommands(const FName InContextName, const FText& InContextDesc, const FName InContextParent, const FName InStyleSetName)
-		: TCommands(InContextName, InContextDesc, InContextParent, InStyleSetName)
+		: TCommands<CommandContextType>(InContextName, InContextDesc, InContextParent, InStyleSetName)
 	{
 	}
 
diff --git a/Engine/Source/Runtime/InteractiveToolsFramework/InteractiveToolsFramework.Build.cs b/Engine/Source/Runtime/InteractiveToolsFramework/InteractiveToolsFramework.Build.cs
--- a/Engine/Source/Runtime/InteractiveToolsFramework/InteractiveToolsFramework.Build.cs
+++ b/Engine/Source/Runtime/InteractiveToolsFramework/InteractiveToolsFramework.Build.cs
@@ -6,7 +6,7 @@ public class InteractiveToolsFramework : ModuleRules
 {
 	public InteractiveToolsFramework(ReadOnlyTargetRules Target) : base(Target)
 	{
-		PCHUsage = ModuleRules.PCHUsageMode.UseExplicitOrSharedPCHs;
+        PCHUsage = ModuleRules.PCHUsageMode.UseExplicitOrSharedPCHs;
 		
 		PublicIncludePaths.AddRange(
 			new string[] {
@@ -25,7 +25,10 @@ public class InteractiveToolsFramework : ModuleRules
 		PublicDependencyModuleNames.AddRange(
 			new string[]
 			{
-				"Core"
+                "Core",
+                "CoreUObject",
+                "InputCore",
+                "ApplicationCore",
 				// ... add other public dependencies that you statically link with here ...
 			}
             );
@@ -34,7 +37,6 @@ public class InteractiveToolsFramework : ModuleRules
 		PrivateDependencyModuleNames.AddRange(
 			new string[]
 			{
-				"CoreUObject",
 				"Engine"
 				//"Slate",
 				//"SlateCore",
