minor improvements to tools framework. MultiClickSequence now has hover and target can abort sequence. Added up/right/forward functions to FViewCameraState.
#rnx

diff --git a/Engine/Source/Runtime/InteractiveToolsFramework/Private/BaseBehaviors/MultiClickSequenceInputBehavior.cpp b/Engine/Source/Runtime/InteractiveToolsFramework/Private/BaseBehaviors/MultiClickSequenceInputBehavior.cpp
--- a/Engine/Source/Runtime/InteractiveToolsFramework/Private/BaseBehaviors/MultiClickSequenceInputBehavior.cpp
+++ b/Engine/Source/Runtime/InteractiveToolsFramework/Private/BaseBehaviors/MultiClickSequenceInputBehavior.cpp
@@ -50,6 +50,14 @@ FInputCaptureUpdate UMultiClickSequenceInputBehavior::UpdateCapture(const FInput
 
 	Modifiers.UpdateModifiers(input, Target);
 
+	// allow target to abort click sequence
+	if (Target->RequestAbortClickSequence())
+	{
+		Target->OnTerminateClickSequence();
+		bInActiveSequence = false;
+		return FInputCaptureUpdate::End();
+	}
+
 	if (IsReleased(input)) 
 	{
 		bool bContinue = Target->OnNextSequenceClick(GetDeviceRay(input));
@@ -73,3 +81,22 @@ void UMultiClickSequenceInputBehavior::ForceEndCapture(const FInputCaptureData&
 	Target->OnTerminateClickSequence();
 	bInActiveSequence = false;
 }
+
+
+bool UMultiClickSequenceInputBehavior::WantsHoverEvents()
+{
+	return true;
+}
+
+void UMultiClickSequenceInputBehavior::UpdateHover(const FInputDeviceState& Input)
+{
+	if (Target != nullptr)
+	{
+		Modifiers.UpdateModifiers(Input, Target);
+		Target->OnBeginSequencePreview(FInputDeviceRay(Input.Mouse.WorldRay, Input.Mouse.Position2D));
+	}
+}
+
+void UMultiClickSequenceInputBehavior::EndHover(const FInputDeviceState& input)
+{
+}
\ No newline at end of file
diff --git a/Engine/Source/Runtime/InteractiveToolsFramework/Public/BaseBehaviors/BehaviorTargetInterfaces.h b/Engine/Source/Runtime/InteractiveToolsFramework/Public/BaseBehaviors/BehaviorTargetInterfaces.h
--- a/Engine/Source/Runtime/InteractiveToolsFramework/Public/BaseBehaviors/BehaviorTargetInterfaces.h
+++ b/Engine/Source/Runtime/InteractiveToolsFramework/Public/BaseBehaviors/BehaviorTargetInterfaces.h
@@ -108,6 +108,12 @@ class IClickSequenceBehaviorTarget : public IModifierToggleBehaviorTarget
 public:
 	virtual ~IClickSequenceBehaviorTarget() {}
 
+	/**
+	 * Notify Target device position has changed but a click sequence hasn't begun yet (eg for interactive previews)
+	 * @param ClickPos device position/ray at click point
+	 */
+	virtual void OnBeginSequencePreview(const FInputDeviceRay& ClickPos) = 0;
+
 	/**
 	 * Test if target would like to begin sequence based on this click
 	 * @param ClickPos device position/ray at click point
@@ -139,6 +145,11 @@ public:
 	 */
 	virtual void OnTerminateClickSequence() = 0;
 
+	/**
+	 * Target overrides this and returns true if it wants to abort click sequence.
+	 * Behavior checks every update and if this ever returns true, terminates sequence
+	 */
+	virtual bool RequestAbortClickSequence() { return false; }
 };
 
 
diff --git a/Engine/Source/Runtime/InteractiveToolsFramework/Public/BaseBehaviors/MultiClickSequenceInputBehavior.h b/Engine/Source/Runtime/InteractiveToolsFramework/Public/BaseBehaviors/MultiClickSequenceInputBehavior.h
--- a/Engine/Source/Runtime/InteractiveToolsFramework/Public/BaseBehaviors/MultiClickSequenceInputBehavior.h
+++ b/Engine/Source/Runtime/InteractiveToolsFramework/Public/BaseBehaviors/MultiClickSequenceInputBehavior.h
@@ -58,6 +58,10 @@ public:
 	virtual FInputCaptureUpdate UpdateCapture(const FInputDeviceState& Input, const FInputCaptureData& Data) override;
 	virtual void ForceEndCapture(const FInputCaptureData& Data) override;
 
+	virtual bool WantsHoverEvents() override;
+	virtual void UpdateHover(const FInputDeviceState& input) override;
+	virtual void EndHover(const FInputDeviceState& input) override;
+
 
 public:
 
diff --git a/Engine/Source/Runtime/InteractiveToolsFramework/Public/ToolContextInterfaces.h b/Engine/Source/Runtime/InteractiveToolsFramework/Public/ToolContextInterfaces.h
--- a/Engine/Source/Runtime/InteractiveToolsFramework/Public/ToolContextInterfaces.h
+++ b/Engine/Source/Runtime/InteractiveToolsFramework/Public/ToolContextInterfaces.h
@@ -56,6 +56,13 @@ struct INTERACTIVETOOLSFRAMEWORK_API FViewCameraState
 	bool bIsOrthographic;
 	/** Is current view a VR view */
 	bool bIsVR;
+
+	/** @return "right"/horizontal direction in camera plane */
+	FVector Right() const { return Orientation.GetAxisY(); }
+	/** @return "up"/vertical direction in camera plane */
+	FVector Up() const { return Orientation.GetAxisZ(); }
+	/** @return forward camera direction */
+	FVector Forward() const { return Orientation.GetAxisX(); }
 };
 
 
@@ -125,6 +132,9 @@ struct INTERACTIVETOOLSFRAMEWORK_API FSceneSnapQueryResult
 
 	/** Vertices of triangle that contains result (for debugging, may not be set) */
 	FVector TriVertices[3];
+	/** Vertex/Edge index we snapped to in triangle */
+	int TriSnapIndex;
+
 };
 
 
