add default implementation of ClickSequenceBehaviorTarget preview functions
#rnx

diff --git a/Engine/Source/Runtime/InteractiveToolsFramework/Public/BaseBehaviors/BehaviorTargetInterfaces.h b/Engine/Source/Runtime/InteractiveToolsFramework/Public/BaseBehaviors/BehaviorTargetInterfaces.h
--- a/Engine/Source/Runtime/InteractiveToolsFramework/Public/BaseBehaviors/BehaviorTargetInterfaces.h
+++ b/Engine/Source/Runtime/InteractiveToolsFramework/Public/BaseBehaviors/BehaviorTargetInterfaces.h
@@ -112,7 +112,7 @@ public:
 	 * Notify Target device position has changed but a click sequence hasn't begun yet (eg for interactive previews)
 	 * @param ClickPos device position/ray at click point
 	 */
-	virtual void OnBeginSequencePreview(const FInputDeviceRay& ClickPos) = 0;
+	virtual void OnBeginSequencePreview(const FInputDeviceRay& ClickPos) { }
 
 	/**
 	 * Test if target would like to begin sequence based on this click
@@ -131,7 +131,7 @@ public:
 	 * Notify Target device position has changed but a click hasn't ocurred yet (eg for interactive previews)
 	 * @param ClickPos device position/ray at click point
 	 */
-	virtual void OnNextSequencePreview(const FInputDeviceRay& ClickPos) = 0;
+	virtual void OnNextSequencePreview(const FInputDeviceRay& ClickPos) { }
 
 	/**
 	 * Notify Target about next click in sqeuence
