add missing header
#rnx

diff --git a/Engine/Source/Runtime/InteractiveToolsFramework/Public/InputState.h b/Engine/Source/Runtime/InteractiveToolsFramework/Public/InputState.h
--- a/Engine/Source/Runtime/InteractiveToolsFramework/Public/InputState.h
+++ b/Engine/Source/Runtime/InteractiveToolsFramework/Public/InputState.h
@@ -4,6 +4,7 @@
 
 #include "CoreMinimal.h"
 #include "UObject/ObjectMacros.h"
+#include "InputCoreTypes.h"
 #include "Math/UnrealMath.h"
 
 
