Use UActorComponent instead of UObject for component in ToolBuilderUtils, ComponentSourceInterfaces.
Return Component world transform instead of Actor world transform in FStaticMeshComponentMeshDescriptionSource::GetWorldTransform, update all DynamicMesh-based tools to set this transform so that nested components appear in right spot.
#rnx

diff --git a/Engine/Source/Editor/EditorInteractiveToolsFramework/Private/EditorComponentSourceFactory.cpp b/Engine/Source/Editor/EditorInteractiveToolsFramework/Private/EditorComponentSourceFactory.cpp
--- a/Engine/Source/Editor/EditorInteractiveToolsFramework/Private/EditorComponentSourceFactory.cpp
+++ b/Engine/Source/Editor/EditorInteractiveToolsFramework/Private/EditorComponentSourceFactory.cpp
@@ -3,12 +3,13 @@
 #include "EditorComponentSourceFactory.h"
 
 #include "Engine/StaticMesh.h"
+#include "Components/ActorComponent.h"
 #include "Components/StaticMeshComponent.h"
 
 
 
 
-TUniquePtr<IMeshDescriptionSource> FEditorComponentSourceFactory::MakeMeshDescriptionSource(UObject* Component)
+TUniquePtr<IMeshDescriptionSource> FEditorComponentSourceFactory::MakeMeshDescriptionSource(UActorComponent* Component)
 {
 	UStaticMeshComponent* StaticMeshComp = Cast<UStaticMeshComponent>(Component);
 	if (StaticMeshComp != nullptr)
@@ -36,7 +37,7 @@ AActor* FStaticMeshComponentMeshDescriptionSource::GetOwnerActor() const
 	return Component->GetOwner();
 }
 
-UObject* FStaticMeshComponentMeshDescriptionSource::GetOwnerComponent() const
+UActorComponent* FStaticMeshComponentMeshDescriptionSource::GetOwnerComponent() const
 {
 	return Component;
 }
@@ -59,7 +60,8 @@ UMaterialInterface* FStaticMeshComponentMeshDescriptionSource::GetMaterial(int32
 
 FTransform FStaticMeshComponentMeshDescriptionSource::GetWorldTransform() const
 {
-	return Component->GetOwner()->GetActorTransform();
+	//return Component->GetOwner()->GetActorTransform();
+	return Component->GetComponentTransform();
 }
 
 
diff --git a/Engine/Source/Editor/EditorInteractiveToolsFramework/Public/EditorComponentSourceFactory.h b/Engine/Source/Editor/EditorInteractiveToolsFramework/Public/EditorComponentSourceFactory.h
--- a/Engine/Source/Editor/EditorInteractiveToolsFramework/Public/EditorComponentSourceFactory.h
+++ b/Engine/Source/Editor/EditorInteractiveToolsFramework/Public/EditorComponentSourceFactory.h
@@ -16,7 +16,7 @@ class EDITORINTERACTIVETOOLSFRAMEWORK_API FEditorComponentSourceFactory : public
 public:
 	virtual ~FEditorComponentSourceFactory() {}
 
-	virtual TUniquePtr<IMeshDescriptionSource> MakeMeshDescriptionSource(UObject* Component);
+	virtual TUniquePtr<IMeshDescriptionSource> MakeMeshDescriptionSource(UActorComponent* Component);
 };
 
 
@@ -40,7 +40,7 @@ public:
 	virtual ~FStaticMeshComponentMeshDescriptionSource() {}
 
 	virtual AActor* GetOwnerActor() const override;
-	virtual UObject* GetOwnerComponent() const override;
+	virtual UActorComponent* GetOwnerComponent() const override;
 	virtual FMeshDescription* GetMeshDescription() const override;
 	virtual UMaterialInterface* GetMaterial(int32 MaterialIndex) const override;
 	virtual FTransform GetWorldTransform() const override;
diff --git a/Engine/Source/Runtime/InteractiveToolsFramework/Private/BaseTools/MeshSurfacePointTool.cpp b/Engine/Source/Runtime/InteractiveToolsFramework/Private/BaseTools/MeshSurfacePointTool.cpp
--- a/Engine/Source/Runtime/InteractiveToolsFramework/Private/BaseTools/MeshSurfacePointTool.cpp
+++ b/Engine/Source/Runtime/InteractiveToolsFramework/Private/BaseTools/MeshSurfacePointTool.cpp
@@ -29,7 +29,7 @@ UMeshSurfacePointTool* UMeshSurfacePointToolBuilder::CreateNewTool(const FToolBu
 
 void UMeshSurfacePointToolBuilder::InitializeNewTool(UMeshSurfacePointTool* NewTool, const FToolBuilderState& SceneState) const
 {
-	UObject* MeshComponent = ToolBuilderUtil::FindFirstComponent(SceneState, ToolBuilderUtil::IsMeshDescriptionSourceComponent);
+	UActorComponent* MeshComponent = ToolBuilderUtil::FindFirstComponent(SceneState, ToolBuilderUtil::IsMeshDescriptionSourceComponent);
 	check(MeshComponent != nullptr);
 	NewTool->SetMeshSource(
 		SceneState.SourceBuilder->MakeMeshDescriptionSource(MeshComponent) );
diff --git a/Engine/Source/Runtime/InteractiveToolsFramework/Private/ToolBuilderUtil.cpp b/Engine/Source/Runtime/InteractiveToolsFramework/Private/ToolBuilderUtil.cpp
--- a/Engine/Source/Runtime/InteractiveToolsFramework/Private/ToolBuilderUtil.cpp
+++ b/Engine/Source/Runtime/InteractiveToolsFramework/Private/ToolBuilderUtil.cpp
@@ -8,7 +8,7 @@
 
 
 
-bool ToolBuilderUtil::IsMeshDescriptionSourceComponent(UObject* ComponentObject)
+bool ToolBuilderUtil::IsMeshDescriptionSourceComponent(UActorComponent* ComponentObject)
 {
 	UStaticMeshComponent* StaticMeshComp = Cast<UStaticMeshComponent>(ComponentObject);
 	if (StaticMeshComp != nullptr)
@@ -20,7 +20,7 @@ bool ToolBuilderUtil::IsMeshDescriptionSourceComponent(UObject* ComponentObject)
 
 
 
-int ToolBuilderUtil::CountComponents(const FToolBuilderState& InputState, const TFunction<bool(UObject*)>& Predicate)
+int ToolBuilderUtil::CountComponents(const FToolBuilderState& InputState, const TFunction<bool(UActorComponent*)>& Predicate)
 {
 	int nTypedComponents = 0;
 
@@ -28,9 +28,12 @@ int ToolBuilderUtil::CountComponents(const FToolBuilderState& InputState, const
 	{
 		for (FSelectionIterator Iter(*InputState.SelectedComponents); Iter; ++Iter)
 		{
-			if (Predicate(*Iter))
+			if (UActorComponent* Component = Cast<UActorComponent>(*Iter))
 			{
-				nTypedComponents++;
+				if (Predicate(Component))
+				{
+					nTypedComponents++;
+				}
 			}
 		}
 	}
@@ -57,15 +60,18 @@ int ToolBuilderUtil::CountComponents(const FToolBuilderState& InputState, const
 
 
 
-UObject* ToolBuilderUtil::FindFirstComponent(const FToolBuilderState& InputState, const TFunction<bool(UObject*)>& Predicate)
+UActorComponent* ToolBuilderUtil::FindFirstComponent(const FToolBuilderState& InputState, const TFunction<bool(UActorComponent*)>& Predicate)
 {
 	if (InputState.SelectedComponents != nullptr && InputState.SelectedComponents->Num() > 0)
 	{
 		for (FSelectionIterator Iter(*InputState.SelectedComponents); Iter; ++Iter)
 		{
-			if (Predicate(*Iter))
+			if (UActorComponent* Component = Cast<UActorComponent>(*Iter))
 			{
-				return *Iter;
+				if (Predicate(Component))
+				{
+					return Component;
+				}
 			}
 		}
 	}
@@ -92,17 +98,20 @@ UObject* ToolBuilderUtil::FindFirstComponent(const FToolBuilderState& InputState
 
 
 
-TArray<UObject*> ToolBuilderUtil::FindAllComponents(const FToolBuilderState& InputState, const TFunction<bool(UObject*)>& Predicate)
+TArray<UActorComponent*> ToolBuilderUtil::FindAllComponents(const FToolBuilderState& InputState, const TFunction<bool(UActorComponent*)>& Predicate)
 {
-	TArray<UObject*> Components;
+	TArray<UActorComponent*> Components;
 
 	if (InputState.SelectedComponents != nullptr && InputState.SelectedComponents->Num() > 0)
 	{
 		for (FSelectionIterator Iter(*InputState.SelectedComponents); Iter; ++Iter)
 		{
-			if (Predicate(*Iter))
+			if (UActorComponent* Component = Cast<UActorComponent>(*Iter))
 			{
-				Components.AddUnique(*Iter);
+				if (Predicate(Component))
+				{
+					Components.AddUnique(Component);
+				}
 			}
 		}
 	}
diff --git a/Engine/Source/Runtime/InteractiveToolsFramework/Public/ComponentSourceInterfaces.h b/Engine/Source/Runtime/InteractiveToolsFramework/Public/ComponentSourceInterfaces.h
--- a/Engine/Source/Runtime/InteractiveToolsFramework/Public/ComponentSourceInterfaces.h
+++ b/Engine/Source/Runtime/InteractiveToolsFramework/Public/ComponentSourceInterfaces.h
@@ -8,6 +8,7 @@
 
 // predeclarations
 class AActor;
+class UActorComponent;
 struct FMeshDescription;
 class UMaterialInterface;
 
@@ -31,7 +32,7 @@ public:
 	virtual AActor* GetOwnerActor() const = 0;
 
 	/** @return the Component this is a Source for */
-	virtual UObject* GetOwnerComponent() const = 0;
+	virtual UActorComponent* GetOwnerComponent() const = 0;
 
 	/** @return Pointer to the MeshDescription this Source is providing */
 	virtual FMeshDescription* GetMeshDescription() const = 0;
@@ -94,6 +95,6 @@ public:
 	 * @param Component A UObject that can provide a MeshDescription. Assumption is this is a Component of an Actor.
 	 * @return A MeshDescriptionSource instance. Must not return null.
 	 */
-	virtual TUniquePtr<IMeshDescriptionSource> MakeMeshDescriptionSource(UObject* Component) = 0;
+	virtual TUniquePtr<IMeshDescriptionSource> MakeMeshDescriptionSource(UActorComponent* Component) = 0;
 };
 
diff --git a/Engine/Source/Runtime/InteractiveToolsFramework/Public/ToolBuilderUtil.h b/Engine/Source/Runtime/InteractiveToolsFramework/Public/ToolBuilderUtil.h
--- a/Engine/Source/Runtime/InteractiveToolsFramework/Public/ToolBuilderUtil.h
+++ b/Engine/Source/Runtime/InteractiveToolsFramework/Public/ToolBuilderUtil.h
@@ -14,19 +14,19 @@ namespace ToolBuilderUtil
 {
 	/** Returns true if this UObject can provide a FMeshDescription */
 	INTERACTIVETOOLSFRAMEWORK_API
-	bool IsMeshDescriptionSourceComponent(UObject* ComponentObject);
+	bool IsMeshDescriptionSourceComponent(UActorComponent* ComponentObject);
 	
 	/** Count number of selected components that pass predicate. If Component selection is not empty, returns that count, otherwise counts in all selected Actors */
 	INTERACTIVETOOLSFRAMEWORK_API
-	int CountComponents(const FToolBuilderState& InputState, const TFunction<bool(UObject*)>& Predicate);
+	int CountComponents(const FToolBuilderState& InputState, const TFunction<bool(UActorComponent*)>& Predicate);
 
 	/** First first available component that passes predicate. Searches Components selection list first, then all Actors */
 	INTERACTIVETOOLSFRAMEWORK_API
-	UObject* FindFirstComponent(const FToolBuilderState& InputState, const TFunction<bool(UObject*)>& Predicate);
+	UActorComponent* FindFirstComponent(const FToolBuilderState& InputState, const TFunction<bool(UActorComponent*)>& Predicate);
 
 	/** First all components that passes predicate. Searches Components selection list first, then all Actors */
 	INTERACTIVETOOLSFRAMEWORK_API
-	TArray<UObject*> FindAllComponents(const FToolBuilderState& InputState, const TFunction<bool(UObject*)>& Predicate);
+	TArray<UActorComponent*> FindAllComponents(const FToolBuilderState& InputState, const TFunction<bool(UActorComponent*)>& Predicate);
 
 	// @todo not sure that actors with multiple components are handled properly...
 	/** Count number of components of given type. If Component selection is not empty, returns that count, otherwise counts in all selected Actors */
