Move InteractiveToolsFramework and GeometryFramework out of Experimental

diff --git a/Engine/Source/Runtime/InteractiveToolsFramework/Private/BaseGizmos/TransformGizmo.cpp b/Engine/Source/Runtime/InteractiveToolsFramework/Private/BaseGizmos/TransformGizmo.cpp
new file mode 100644
--- /dev/null
+++ b/Engine/Source/Runtime/InteractiveToolsFramework/Private/BaseGizmos/TransformGizmo.cpp
@@ -0,0 +1,978 @@
+// Copyright Epic Games, Inc. All Rights Reserved.
+
+#include "BaseGizmos/TransformGizmo.h"
+#include "InteractiveGizmoManager.h"
+#include "BaseGizmos/AxisPositionGizmo.h"
+#include "BaseGizmos/PlanePositionGizmo.h"
+#include "BaseGizmos/AxisAngleGizmo.h"
+#include "BaseGizmos/GizmoComponents.h"
+
+#include "BaseGizmos/GizmoArrowComponent.h"
+#include "BaseGizmos/GizmoRectangleComponent.h"
+#include "BaseGizmos/GizmoCircleComponent.h"
+#include "BaseGizmos/GizmoBoxComponent.h"
+#include "BaseGizmos/GizmoLineHandleComponent.h"
+#include "BaseGizmos/GizmoViewContext.h"
+
+// need this to implement hover
+#include "BaseGizmos/GizmoBaseComponent.h"
+
+#include "Components/SphereComponent.h"
+#include "Components/PrimitiveComponent.h"
+#include "ContextObjectStore.h"
+#include "Engine/World.h"
+#include "Engine/CollisionProfile.h"
+
+
+#define LOCTEXT_NAMESPACE "UTransformGizmo"
+
+
+ATransformGizmoActor::ATransformGizmoActor()
+{
+	// root component is a hidden sphere
+	USphereComponent* SphereComponent = CreateDefaultSubobject<USphereComponent>(TEXT("GizmoCenter"));
+	RootComponent = SphereComponent;
+	SphereComponent->InitSphereRadius(1.0f);
+	SphereComponent->SetVisibility(false);
+	SphereComponent->SetCollisionProfileName(UCollisionProfile::NoCollision_ProfileName);
+}
+
+
+
+ATransformGizmoActor* ATransformGizmoActor::ConstructDefault3AxisGizmo(UWorld* World, UGizmoViewContext* GizmoViewContext)
+{
+	return ConstructCustom3AxisGizmo(World, GizmoViewContext,
+		ETransformGizmoSubElements::TranslateAllAxes |
+		ETransformGizmoSubElements::TranslateAllPlanes |
+		ETransformGizmoSubElements::RotateAllAxes |
+		ETransformGizmoSubElements::ScaleAllAxes |
+		ETransformGizmoSubElements::ScaleAllPlanes |
+		ETransformGizmoSubElements::ScaleUniform
+	);
+}
+
+
+ATransformGizmoActor* ATransformGizmoActor::ConstructCustom3AxisGizmo(
+	UWorld* World, UGizmoViewContext* GizmoViewContext,
+	ETransformGizmoSubElements Elements)
+{
+	FActorSpawnParameters SpawnInfo;
+	ATransformGizmoActor* NewActor = World->SpawnActor<ATransformGizmoActor>(FVector::ZeroVector, FRotator::ZeroRotator, SpawnInfo);
+
+	float GizmoLineThickness = 3.0f;
+
+
+	auto MakeAxisArrowFunc = [&](const FLinearColor& Color, const FVector& Axis)
+	{
+		UGizmoArrowComponent* Component = AddDefaultArrowComponent(World, NewActor, GizmoViewContext, Color, Axis, 60.0f);
+		Component->Gap = 20.0f;
+		Component->Thickness = GizmoLineThickness;
+		Component->NotifyExternalPropertyUpdates();
+		return Component;
+	};
+	if ((Elements & ETransformGizmoSubElements::TranslateAxisX) != ETransformGizmoSubElements::None)
+	{
+		NewActor->TranslateX = MakeAxisArrowFunc(FLinearColor::Red, FVector(1, 0, 0));
+	}
+	if ((Elements & ETransformGizmoSubElements::TranslateAxisY) != ETransformGizmoSubElements::None)
+	{
+		NewActor->TranslateY = MakeAxisArrowFunc(FLinearColor::Green, FVector(0, 1, 0));
+	}
+	if ((Elements & ETransformGizmoSubElements::TranslateAxisZ) != ETransformGizmoSubElements::None)
+	{
+		NewActor->TranslateZ = MakeAxisArrowFunc(FLinearColor::Blue, FVector(0, 0, 1));
+	}
+
+
+	auto MakePlaneRectFunc = [&](const FLinearColor& Color, const FVector& Axis0, const FVector& Axis1)
+	{
+		UGizmoRectangleComponent* Component = AddDefaultRectangleComponent(World, NewActor, GizmoViewContext, Color, Axis0, Axis1);
+		Component->LengthX = Component->LengthY = 30.0f;
+		Component->SegmentFlags = 0x2 | 0x4;
+		Component->Thickness = GizmoLineThickness;
+		Component->NotifyExternalPropertyUpdates();
+		return Component;
+	};
+	if ((Elements & ETransformGizmoSubElements::TranslatePlaneYZ) != ETransformGizmoSubElements::None)
+	{
+		NewActor->TranslateYZ = MakePlaneRectFunc(FLinearColor::Red, FVector(0, 1, 0), FVector(0, 0, 1));
+	}
+	if ((Elements & ETransformGizmoSubElements::TranslatePlaneXZ) != ETransformGizmoSubElements::None)
+	{
+		NewActor->TranslateXZ = MakePlaneRectFunc(FLinearColor::Green, FVector(1, 0, 0), FVector(0, 0, 1));
+	}
+	if ((Elements & ETransformGizmoSubElements::TranslatePlaneXY) != ETransformGizmoSubElements::None)
+	{
+		NewActor->TranslateXY = MakePlaneRectFunc(FLinearColor::Blue, FVector(1, 0, 0), FVector(0, 1, 0));
+	}
+
+	auto MakeAxisRotateCircleFunc = [&](const FLinearColor& Color, const FVector& Axis)
+	{
+		UGizmoCircleComponent* Component = AddDefaultCircleComponent(World, NewActor, GizmoViewContext, Color, Axis, 120.0f);
+		Component->Thickness = GizmoLineThickness;
+		Component->NotifyExternalPropertyUpdates();
+		return Component;
+	};
+
+	bool bAnyRotate = false;
+	if ((Elements & ETransformGizmoSubElements::RotateAxisX) != ETransformGizmoSubElements::None)
+	{
+		NewActor->RotateX = MakeAxisRotateCircleFunc(FLinearColor::Red, FVector(1, 0, 0));
+		bAnyRotate = true;
+	}
+	if ((Elements & ETransformGizmoSubElements::RotateAxisY) != ETransformGizmoSubElements::None)
+	{
+		NewActor->RotateY = MakeAxisRotateCircleFunc(FLinearColor::Green, FVector(0, 1, 0));
+		bAnyRotate = true;
+	}
+	if ((Elements & ETransformGizmoSubElements::RotateAxisZ) != ETransformGizmoSubElements::None)
+	{
+		NewActor->RotateZ = MakeAxisRotateCircleFunc(FLinearColor::Blue, FVector(0, 0, 1));
+		bAnyRotate = true;
+	}
+
+
+	// add a non-interactive view-aligned circle element, so the axes look like a sphere.
+	if (bAnyRotate)
+	{
+		UGizmoCircleComponent* SphereEdge = NewObject<UGizmoCircleComponent>(NewActor);
+		NewActor->AddInstanceComponent(SphereEdge);
+		SphereEdge->AttachToComponent(NewActor->GetRootComponent(), FAttachmentTransformRules::KeepRelativeTransform);
+		SphereEdge->SetGizmoViewContext(GizmoViewContext);
+		SphereEdge->Color = FLinearColor::Gray;
+		SphereEdge->Thickness = 1.0f;
+		SphereEdge->Radius = 120.0f;
+		SphereEdge->bViewAligned = true;
+		SphereEdge->RegisterComponent();
+	}
+
+
+
+	if ((Elements & ETransformGizmoSubElements::ScaleUniform) != ETransformGizmoSubElements::None)
+	{
+		float BoxSize = 14.0f;
+		UGizmoBoxComponent* ScaleComponent = AddDefaultBoxComponent(World, NewActor, GizmoViewContext, FLinearColor::Black,
+			FVector(BoxSize/2, BoxSize/2, BoxSize/2), FVector(BoxSize, BoxSize, BoxSize));
+		NewActor->UniformScale = ScaleComponent;
+	}
+
+
+
+	auto MakeAxisScaleFunc = [&](const FLinearColor& Color, const FVector& Axis0, const FVector& Axis1)
+	{
+		UGizmoRectangleComponent* ScaleComponent = AddDefaultRectangleComponent(World, NewActor, GizmoViewContext, Color, Axis0, Axis1);
+		ScaleComponent->OffsetX = 140.0f; ScaleComponent->OffsetY = -10.0f;
+		ScaleComponent->LengthX = 7.0f; ScaleComponent->LengthY = 20.0f;
+		ScaleComponent->Thickness = GizmoLineThickness;
+		ScaleComponent->bOrientYAccordingToCamera = true;
+		ScaleComponent->NotifyExternalPropertyUpdates();
+		ScaleComponent->SegmentFlags = 0x1 | 0x2 | 0x4; // | 0x8;
+		return ScaleComponent;
+	};
+	if ((Elements & ETransformGizmoSubElements::ScaleAxisX) != ETransformGizmoSubElements::None)
+	{
+		NewActor->AxisScaleX = MakeAxisScaleFunc(FLinearColor::Red, FVector(1, 0, 0), FVector(0, 0, 1));
+	}
+	if ((Elements & ETransformGizmoSubElements::ScaleAxisY) != ETransformGizmoSubElements::None)
+	{
+		NewActor->AxisScaleY = MakeAxisScaleFunc(FLinearColor::Green, FVector(0, 1, 0), FVector(0, 0, 1));
+	}
+	if ((Elements & ETransformGizmoSubElements::ScaleAxisZ) != ETransformGizmoSubElements::None)
+	{
+		NewActor->AxisScaleZ = MakeAxisScaleFunc(FLinearColor::Blue, FVector(0, 0, 1), FVector(1, 0, 0));
+	}
+
+
+	auto MakePlaneScaleFunc = [&](const FLinearColor& Color, const FVector& Axis0, const FVector& Axis1)
+	{
+		UGizmoRectangleComponent* ScaleComponent = AddDefaultRectangleComponent(World, NewActor, GizmoViewContext, Color, Axis0, Axis1);
+		ScaleComponent->OffsetX = ScaleComponent->OffsetY = 120.0f;
+		ScaleComponent->LengthX = ScaleComponent->LengthY = 20.0f;
+		ScaleComponent->Thickness = GizmoLineThickness;
+		ScaleComponent->NotifyExternalPropertyUpdates();
+		ScaleComponent->SegmentFlags = 0x2 | 0x4;
+		return ScaleComponent;
+	};
+	if ((Elements & ETransformGizmoSubElements::ScalePlaneYZ) != ETransformGizmoSubElements::None)
+	{
+		NewActor->PlaneScaleYZ = MakePlaneScaleFunc(FLinearColor::Red, FVector(0, 1, 0), FVector(0, 0, 1));
+	}
+	if ((Elements & ETransformGizmoSubElements::ScalePlaneXZ) != ETransformGizmoSubElements::None)
+	{
+		NewActor->PlaneScaleXZ = MakePlaneScaleFunc(FLinearColor::Green, FVector(1, 0, 0), FVector(0, 0, 1));
+	}
+	if ((Elements & ETransformGizmoSubElements::ScalePlaneXY) != ETransformGizmoSubElements::None)
+	{
+		NewActor->PlaneScaleXY = MakePlaneScaleFunc(FLinearColor::Blue, FVector(1, 0, 0), FVector(0, 1, 0));
+	}
+
+
+	return NewActor;
+}
+
+
+
+
+ATransformGizmoActor* FTransformGizmoActorFactory::CreateNewGizmoActor(UWorld* World) const
+{
+	return ATransformGizmoActor::ConstructCustom3AxisGizmo(World, GizmoViewContext, EnableElements);
+}
+
+
+
+UInteractiveGizmo* UTransformGizmoBuilder::BuildGizmo(const FToolBuilderState& SceneState) const
+{
+	UTransformGizmo* NewGizmo = NewObject<UTransformGizmo>(SceneState.GizmoManager);
+	NewGizmo->SetWorld(SceneState.World);
+
+	UGizmoViewContext* GizmoViewContext = SceneState.ToolManager->GetContextObjectStore()->FindContext<UGizmoViewContext>();
+	check(GizmoViewContext && GizmoViewContext->IsValidLowLevel());
+
+	// use default gizmo actor if client has not given us a new builder
+	NewGizmo->SetGizmoActorBuilder(GizmoActorBuilder ? GizmoActorBuilder : MakeShared<FTransformGizmoActorFactory>(GizmoViewContext));
+
+	NewGizmo->SetSubGizmoBuilderIdentifiers(AxisPositionBuilderIdentifier, PlanePositionBuilderIdentifier, AxisAngleBuilderIdentifier);
+
+	// override default hover function if proposed
+	if (UpdateHoverFunction)
+	{
+		NewGizmo->SetUpdateHoverFunction(UpdateHoverFunction);
+	}
+
+	if (UpdateCoordSystemFunction)
+	{
+		NewGizmo->SetUpdateCoordSystemFunction(UpdateCoordSystemFunction);
+	}
+
+	return NewGizmo;
+}
+
+
+
+void UTransformGizmo::SetWorld(UWorld* WorldIn)
+{
+	this->World = WorldIn;
+}
+
+void UTransformGizmo::SetGizmoActorBuilder(TSharedPtr<FTransformGizmoActorFactory> Builder)
+{
+	GizmoActorBuilder = Builder;
+}
+
+void UTransformGizmo::SetSubGizmoBuilderIdentifiers(FString AxisPositionBuilderIdentifierIn, FString PlanePositionBuilderIdentifierIn, FString AxisAngleBuilderIdentifierIn)
+{
+	AxisPositionBuilderIdentifier = AxisPositionBuilderIdentifierIn;
+	PlanePositionBuilderIdentifier = PlanePositionBuilderIdentifierIn;
+	AxisAngleBuilderIdentifier = AxisAngleBuilderIdentifierIn;
+}
+
+void UTransformGizmo::SetUpdateHoverFunction(TFunction<void(UPrimitiveComponent*, bool)> HoverFunction)
+{
+	UpdateHoverFunction = HoverFunction;
+}
+
+void UTransformGizmo::SetUpdateCoordSystemFunction(TFunction<void(UPrimitiveComponent*, EToolContextCoordinateSystem)> CoordSysFunction)
+{
+	UpdateCoordSystemFunction = CoordSysFunction;
+}
+
+void UTransformGizmo::SetWorldAlignmentFunctions(
+	TUniqueFunction<bool()>&& ShouldAlignTranslationIn,
+	TUniqueFunction<bool(const FRay&, FVector&)>&& TranslationAlignmentRayCasterIn)
+{
+	// Save these so that later changes of gizmo target keep the settings.
+	ShouldAlignDestination = MoveTemp(ShouldAlignTranslationIn);
+	DestinationAlignmentRayCaster = MoveTemp(TranslationAlignmentRayCasterIn);
+
+	// We allow this function to be called after Setup(), so modify any existing translation/rotation sub gizmos.
+	// Unfortunately we keep all the sub gizmos in one list, and the scaling gizmos are differentiated from the
+	// translation ones mainly in the components they use. So this ends up being a slightly messy set of checks,
+	// but it didn't seem worth keeping a segregated list for something that will only happen once.
+	for (UInteractiveGizmo* SubGizmo : this->ActiveGizmos)
+	{
+		if (UAxisPositionGizmo* CastGizmo = Cast<UAxisPositionGizmo>(SubGizmo))
+		{
+			if (UGizmoComponentHitTarget* CastHitTarget = Cast<UGizmoComponentHitTarget>(CastGizmo->HitTarget.GetObject()))
+			{
+				if (CastHitTarget->Component == GizmoActor->TranslateX
+					|| CastHitTarget->Component == GizmoActor->TranslateY
+					|| CastHitTarget->Component == GizmoActor->TranslateZ)
+				{
+					CastGizmo->ShouldUseCustomDestinationFunc = [this]() { return ShouldAlignDestination(); };
+					CastGizmo->CustomDestinationFunc = 
+						[this](const UAxisPositionGizmo::FCustomDestinationParams& Params, FVector& OutputPoint) { 
+						return DestinationAlignmentRayCaster(*Params.WorldRay, OutputPoint); 
+					};
+				}
+			}
+		}
+		if (UPlanePositionGizmo* CastGizmo = Cast<UPlanePositionGizmo>(SubGizmo))
+		{
+			if (UGizmoComponentHitTarget* CastHitTarget = Cast<UGizmoComponentHitTarget>(CastGizmo->HitTarget.GetObject()))
+			{
+				if (CastHitTarget->Component == GizmoActor->TranslateXY
+					|| CastHitTarget->Component == GizmoActor->TranslateXZ
+					|| CastHitTarget->Component == GizmoActor->TranslateYZ)
+				{
+					CastGizmo->ShouldUseCustomDestinationFunc = [this]() { return ShouldAlignDestination(); };
+					CastGizmo->CustomDestinationFunc =
+						[this](const UPlanePositionGizmo::FCustomDestinationParams& Params, FVector& OutputPoint) {
+						return DestinationAlignmentRayCaster(*Params.WorldRay, OutputPoint);
+					};
+				}
+			}
+		}
+		if (UAxisAngleGizmo* CastGizmo = Cast<UAxisAngleGizmo>(SubGizmo))
+		{
+			CastGizmo->ShouldUseCustomDestinationFunc = [this]() { return ShouldAlignDestination(); };
+			CastGizmo->CustomDestinationFunc =
+				[this](const UAxisAngleGizmo::FCustomDestinationParams& Params, FVector& OutputPoint) {
+				return DestinationAlignmentRayCaster(*Params.WorldRay, OutputPoint);
+			};
+		}
+	}
+}
+
+void UTransformGizmo::SetDisallowNegativeScaling(bool bDisallow)
+{
+	if (bDisallowNegativeScaling != bDisallow)
+	{
+		bDisallowNegativeScaling = bDisallow;
+		for (UInteractiveGizmo* SubGizmo : this->ActiveGizmos)
+		{
+			if (UAxisPositionGizmo* CastGizmo = Cast<UAxisPositionGizmo>(SubGizmo))
+			{
+				if (UGizmoAxisScaleParameterSource* ParamSource = Cast<UGizmoAxisScaleParameterSource>(CastGizmo->ParameterSource.GetObject()))
+				{
+					ParamSource->bClampToZero = bDisallow;
+				}
+			}
+			if (UPlanePositionGizmo* CastGizmo = Cast<UPlanePositionGizmo>(SubGizmo))
+			{
+				if (UGizmoPlaneScaleParameterSource* ParamSource = Cast<UGizmoPlaneScaleParameterSource>(CastGizmo->ParameterSource.GetObject()))
+				{
+					ParamSource->bClampToZero = bDisallow;
+				}
+			}
+		}
+	}
+}
+
+void UTransformGizmo::SetIsNonUniformScaleAllowedFunction(TUniqueFunction<bool()>&& IsNonUniformScaleAllowedIn)
+{
+	IsNonUniformScaleAllowed = MoveTemp(IsNonUniformScaleAllowedIn);
+}
+
+
+void UTransformGizmo::Setup()
+{
+	UInteractiveGizmo::Setup();
+
+	UpdateHoverFunction = [](UPrimitiveComponent* Component, bool bHovering)
+	{
+		if (Cast<UGizmoBaseComponent>(Component) != nullptr)
+		{
+			Cast<UGizmoBaseComponent>(Component)->UpdateHoverState(bHovering);
+		}
+	};
+
+	UpdateCoordSystemFunction = [](UPrimitiveComponent* Component, EToolContextCoordinateSystem CoordSystem)
+	{
+		if (Cast<UGizmoBaseComponent>(Component) != nullptr)
+		{
+			Cast<UGizmoBaseComponent>(Component)->UpdateWorldLocalState(CoordSystem == EToolContextCoordinateSystem::World);
+		}
+	};
+
+	GizmoActor = GizmoActorBuilder->CreateNewGizmoActor(World);
+}
+
+
+
+void UTransformGizmo::Shutdown()
+{
+	ClearActiveTarget();
+
+	if (GizmoActor)
+	{
+		GizmoActor->Destroy();
+		GizmoActor = nullptr;
+	}
+}
+
+
+
+void UTransformGizmo::UpdateCameraAxisSource()
+{
+	FViewCameraState CameraState;
+	GetGizmoManager()->GetContextQueriesAPI()->GetCurrentViewState(CameraState);
+	if (CameraAxisSource != nullptr && GizmoActor != nullptr)
+	{
+		CameraAxisSource->Origin = GizmoActor->GetTransform().GetLocation();
+		CameraAxisSource->Direction = -CameraState.Forward();
+		CameraAxisSource->TangentX = CameraState.Right();
+		CameraAxisSource->TangentY = CameraState.Up();
+	}
+}
+
+
+void UTransformGizmo::Tick(float DeltaTime)
+{	
+	if (bUseContextCoordinateSystem)
+	{
+		CurrentCoordinateSystem = GetGizmoManager()->GetContextQueriesAPI()->GetCurrentCoordinateSystem();
+	}
+	
+	check(CurrentCoordinateSystem == EToolContextCoordinateSystem::World || CurrentCoordinateSystem == EToolContextCoordinateSystem::Local)
+	bool bUseLocalAxes = (CurrentCoordinateSystem == EToolContextCoordinateSystem::Local);
+
+	if (AxisXSource != nullptr && AxisYSource != nullptr && AxisZSource != nullptr)
+	{
+		AxisXSource->bLocalAxes = bUseLocalAxes;
+		AxisYSource->bLocalAxes = bUseLocalAxes;
+		AxisZSource->bLocalAxes = bUseLocalAxes;
+	}
+	if (UpdateCoordSystemFunction)
+	{
+		for (UPrimitiveComponent* Component : ActiveComponents)
+		{
+			UpdateCoordSystemFunction(Component, CurrentCoordinateSystem);
+		}
+	}
+
+	bool bShouldShowNonUniformScale = IsNonUniformScaleAllowed();
+	for (UPrimitiveComponent* Component : NonuniformScaleComponents)
+	{
+		Component->SetVisibility(bShouldShowNonUniformScale);
+	}
+
+	UpdateCameraAxisSource();
+}
+
+
+
+void UTransformGizmo::SetActiveTarget(UTransformProxy* Target, IToolContextTransactionProvider* TransactionProvider)
+{
+	if (ActiveTarget != nullptr)
+	{
+		ClearActiveTarget();
+	}
+
+	ActiveTarget = Target;
+
+	// move gizmo to target location
+	USceneComponent* GizmoComponent = GizmoActor->GetRootComponent();
+
+	FTransform TargetTransform = Target->GetTransform();
+	FTransform GizmoTransform = TargetTransform;
+	GizmoTransform.SetScale3D(FVector(1, 1, 1));
+	GizmoComponent->SetWorldTransform(GizmoTransform);
+
+	UGizmoScaledAndUnscaledTransformSources* TransformSource = UGizmoScaledAndUnscaledTransformSources::Construct(
+		UGizmoTransformProxyTransformSource::Construct(ActiveTarget, this), 
+		GizmoComponent, this);
+
+	// This state target emits an explicit FChange that moves the GizmoActor root component during undo/redo.
+	// It also opens/closes the Transaction that saves/restores the target object locations.
+	if (TransactionProvider == nullptr)
+	{
+		TransactionProvider = GetGizmoManager();
+	}
+	StateTarget = UGizmoTransformChangeStateTarget::Construct(GizmoComponent,
+		LOCTEXT("UTransformGizmoTransaction", "Transform"), TransactionProvider, this);
+	StateTarget->DependentChangeSources.Add(MakeUnique<FTransformProxyChangeSource>(Target));
+
+	CameraAxisSource = NewObject<UGizmoConstantFrameAxisSource>(this);
+
+	// root component provides local X/Y/Z axis, identified by AxisIndex
+	AxisXSource = UGizmoComponentAxisSource::Construct(GizmoComponent, 0, true, this);
+	AxisYSource = UGizmoComponentAxisSource::Construct(GizmoComponent, 1, true, this);
+	AxisZSource = UGizmoComponentAxisSource::Construct(GizmoComponent, 2, true, this);
+
+	// todo should we hold onto these?
+	if (GizmoActor->TranslateX != nullptr)
+	{
+		AddAxisTranslationGizmo(GizmoActor->TranslateX, GizmoComponent, AxisXSource, TransformSource, StateTarget);
+		ActiveComponents.Add(GizmoActor->TranslateX);
+	}
+	if (GizmoActor->TranslateY != nullptr)
+	{
+		AddAxisTranslationGizmo(GizmoActor->TranslateY, GizmoComponent, AxisYSource, TransformSource, StateTarget);
+		ActiveComponents.Add(GizmoActor->TranslateY);
+	}
+	if (GizmoActor->TranslateZ != nullptr)
+	{
+		AddAxisTranslationGizmo(GizmoActor->TranslateZ, GizmoComponent, AxisZSource, TransformSource, StateTarget);
+		ActiveComponents.Add(GizmoActor->TranslateZ);
+	}
+
+
+	if (GizmoActor->TranslateYZ != nullptr)
+	{
+		AddPlaneTranslationGizmo(GizmoActor->TranslateYZ, GizmoComponent, AxisXSource, TransformSource, StateTarget);
+		ActiveComponents.Add(GizmoActor->TranslateYZ);
+	}
+	if (GizmoActor->TranslateXZ != nullptr)
+	{
+		AddPlaneTranslationGizmo(GizmoActor->TranslateXZ, GizmoComponent, AxisYSource, TransformSource, StateTarget);
+		ActiveComponents.Add(GizmoActor->TranslateXZ);
+	}
+	if (GizmoActor->TranslateXY != nullptr)
+	{
+		AddPlaneTranslationGizmo(GizmoActor->TranslateXY, GizmoComponent, AxisZSource, TransformSource, StateTarget);
+		ActiveComponents.Add(GizmoActor->TranslateXY);
+	}
+
+	if (GizmoActor->RotateX != nullptr)
+	{
+		AddAxisRotationGizmo(GizmoActor->RotateX, GizmoComponent, AxisXSource, TransformSource, StateTarget);
+		ActiveComponents.Add(GizmoActor->RotateX);
+	}
+	if (GizmoActor->RotateY != nullptr)
+	{
+		AddAxisRotationGizmo(GizmoActor->RotateY, GizmoComponent, AxisYSource, TransformSource, StateTarget);
+		ActiveComponents.Add(GizmoActor->RotateY);
+	}
+	if (GizmoActor->RotateZ != nullptr)
+	{
+		AddAxisRotationGizmo(GizmoActor->RotateZ, GizmoComponent, AxisZSource, TransformSource, StateTarget);
+		ActiveComponents.Add(GizmoActor->RotateZ);
+	}
+
+
+	// only need these if scaling enabled. Essentially these are just the unit axes, regardless
+	// of what 3D axis is in use, we will tell the ParameterSource-to-3D-Scale mapper to
+	// use the coordinate axes
+	UnitAxisXSource = UGizmoComponentAxisSource::Construct(GizmoComponent, 0, false, this);
+	UnitAxisYSource = UGizmoComponentAxisSource::Construct(GizmoComponent, 1, false, this);
+	UnitAxisZSource = UGizmoComponentAxisSource::Construct(GizmoComponent, 2, false, this);
+
+	if (GizmoActor->UniformScale != nullptr)
+	{
+		AddUniformScaleGizmo(GizmoActor->UniformScale, GizmoComponent, CameraAxisSource, CameraAxisSource, TransformSource, StateTarget);
+		ActiveComponents.Add(GizmoActor->UniformScale);
+	}
+
+	if (GizmoActor->AxisScaleX != nullptr)
+	{
+		AddAxisScaleGizmo(GizmoActor->AxisScaleX, GizmoComponent, AxisXSource, UnitAxisXSource, TransformSource, StateTarget);
+		ActiveComponents.Add(GizmoActor->AxisScaleX);
+		NonuniformScaleComponents.Add(GizmoActor->AxisScaleX);
+	}
+	if (GizmoActor->AxisScaleY != nullptr)
+	{
+		AddAxisScaleGizmo(GizmoActor->AxisScaleY, GizmoComponent, AxisYSource, UnitAxisYSource, TransformSource, StateTarget);
+		ActiveComponents.Add(GizmoActor->AxisScaleY);
+		NonuniformScaleComponents.Add(GizmoActor->AxisScaleY);
+	}
+	if (GizmoActor->AxisScaleZ != nullptr)
+	{
+		AddAxisScaleGizmo(GizmoActor->AxisScaleZ, GizmoComponent, AxisZSource, UnitAxisZSource, TransformSource, StateTarget);
+		ActiveComponents.Add(GizmoActor->AxisScaleZ);
+		NonuniformScaleComponents.Add(GizmoActor->AxisScaleZ);
+	}
+
+	if (GizmoActor->PlaneScaleYZ != nullptr)
+	{
+		AddPlaneScaleGizmo(GizmoActor->PlaneScaleYZ, GizmoComponent, AxisXSource, UnitAxisXSource, TransformSource, StateTarget);
+		ActiveComponents.Add(GizmoActor->PlaneScaleYZ);
+		NonuniformScaleComponents.Add(GizmoActor->PlaneScaleYZ);
+	}
+	if (GizmoActor->PlaneScaleXZ != nullptr)
+	{
+		UPlanePositionGizmo* Gizmo = (UPlanePositionGizmo *)AddPlaneScaleGizmo(GizmoActor->PlaneScaleXZ, GizmoComponent, AxisYSource, UnitAxisYSource, TransformSource, StateTarget);
+		ActiveComponents.Add(GizmoActor->PlaneScaleXZ);
+		NonuniformScaleComponents.Add(GizmoActor->PlaneScaleXZ);
+	}
+	if (GizmoActor->PlaneScaleXY != nullptr)
+	{
+		AddPlaneScaleGizmo(GizmoActor->PlaneScaleXY, GizmoComponent, AxisZSource, UnitAxisZSource, TransformSource, StateTarget);
+		ActiveComponents.Add(GizmoActor->PlaneScaleXY);
+		NonuniformScaleComponents.Add(GizmoActor->PlaneScaleXY);
+	}
+}
+
+FTransform UTransformGizmo::GetGizmoTransform() const
+{
+	USceneComponent* GizmoComponent = GizmoActor->GetRootComponent();
+	return GizmoComponent->GetComponentTransform();
+}
+
+void UTransformGizmo::ReinitializeGizmoTransform(const FTransform& NewTransform, bool bKeepGizmoUnscaled)
+{
+	// To update the gizmo location without triggering any callbacks, we temporarily
+	// store a copy of the callback list, detach them, reposition, and then reattach
+	// the callbacks.
+	USceneComponent* GizmoComponent = GizmoActor->GetRootComponent();
+	auto temp = GizmoComponent->TransformUpdated;
+	GizmoComponent->TransformUpdated.Clear();
+	FTransform GizmoTransform = NewTransform;
+	if (bKeepGizmoUnscaled)
+	{
+		GizmoTransform.SetScale3D(FVector(1, 1, 1));
+	}
+	GizmoComponent->SetWorldTransform(GizmoTransform);
+	GizmoComponent->TransformUpdated = temp;
+
+	// The underlying proxy has an existing way to reinitialize its transform without callbacks.
+	bool bSavedSetPivotMode = ActiveTarget->bSetPivotMode;
+	ActiveTarget->bSetPivotMode = true;
+	ActiveTarget->SetTransform(NewTransform);
+	ActiveTarget->bSetPivotMode = bSavedSetPivotMode;
+}
+
+void UTransformGizmo::SetNewGizmoTransform(const FTransform& NewTransform, bool bKeepGizmoUnscaled)
+{
+	check(ActiveTarget != nullptr);
+
+	StateTarget->BeginUpdate();
+
+	USceneComponent* GizmoComponent = GizmoActor->GetRootComponent();
+	FTransform GizmoTransform = NewTransform;
+	if (bKeepGizmoUnscaled)
+	{
+		GizmoTransform.SetScale3D(FVector(1, 1, 1));
+	}
+	GizmoComponent->SetWorldTransform(GizmoTransform);
+	ActiveTarget->SetTransform(NewTransform);
+
+	StateTarget->EndUpdate();
+}
+
+
+void UTransformGizmo::SetNewChildScale(const FVector& NewChildScale)
+{
+	FTransform NewTransform = ActiveTarget->GetTransform();
+	NewTransform.SetScale3D(NewChildScale);
+
+	bool bSavedSetPivotMode = ActiveTarget->bSetPivotMode;
+	ActiveTarget->bSetPivotMode = true;
+	ActiveTarget->SetTransform(NewTransform);
+	ActiveTarget->bSetPivotMode = bSavedSetPivotMode;
+}
+
+
+void UTransformGizmo::SetVisibility(bool bVisible)
+{
+	GizmoActor->SetActorHiddenInGame(bVisible == false);
+#if WITH_EDITOR
+	GizmoActor->SetIsTemporarilyHiddenInEditor(bVisible == false);
+#endif
+}
+
+
+UInteractiveGizmo* UTransformGizmo::AddAxisTranslationGizmo(
+	UPrimitiveComponent* AxisComponent, USceneComponent* RootComponent,
+	IGizmoAxisSource* AxisSource,
+	IGizmoTransformSource* TransformSource,
+	IGizmoStateTarget* StateTargetIn)
+{
+	// create axis-position gizmo, axis-position parameter will drive translation
+	UAxisPositionGizmo* TranslateGizmo = Cast<UAxisPositionGizmo>(GetGizmoManager()->CreateGizmo(AxisPositionBuilderIdentifier));
+	check(TranslateGizmo);
+
+	// axis source provides the translation axis
+	TranslateGizmo->AxisSource = Cast<UObject>(AxisSource);
+
+	// parameter source maps axis-parameter-change to translation of TransformSource's transform
+	UGizmoAxisTranslationParameterSource* ParamSource = UGizmoAxisTranslationParameterSource::Construct(AxisSource, TransformSource, this);
+	ParamSource->PositionConstraintFunction = [this](const FVector& Pos, FVector& Snapped) { return PositionSnapFunction(Pos, Snapped); };
+	TranslateGizmo->ParameterSource = ParamSource;
+
+	// sub-component provides hit target
+	UGizmoComponentHitTarget* HitTarget = UGizmoComponentHitTarget::Construct(AxisComponent, this);
+	if (this->UpdateHoverFunction)
+	{
+		HitTarget->UpdateHoverFunction = [AxisComponent, this](bool bHovering) { this->UpdateHoverFunction(AxisComponent, bHovering); };
+	}
+	TranslateGizmo->HitTarget = HitTarget;
+
+	TranslateGizmo->StateTarget = Cast<UObject>(StateTargetIn);
+
+	TranslateGizmo->ShouldUseCustomDestinationFunc = [this]() { return ShouldAlignDestination(); };
+	TranslateGizmo->CustomDestinationFunc =
+		[this](const UAxisPositionGizmo::FCustomDestinationParams& Params, FVector& OutputPoint) {
+		return DestinationAlignmentRayCaster(*Params.WorldRay, OutputPoint);
+	};
+
+	ActiveGizmos.Add(TranslateGizmo);
+	return TranslateGizmo;
+}
+
+
+
+UInteractiveGizmo* UTransformGizmo::AddPlaneTranslationGizmo(
+	UPrimitiveComponent* AxisComponent, USceneComponent* RootComponent,
+	IGizmoAxisSource* AxisSource,
+	IGizmoTransformSource* TransformSource,
+	IGizmoStateTarget* StateTargetIn)
+{
+	// create axis-position gizmo, axis-position parameter will drive translation
+	UPlanePositionGizmo* TranslateGizmo = Cast<UPlanePositionGizmo>(GetGizmoManager()->CreateGizmo(PlanePositionBuilderIdentifier));
+	check(TranslateGizmo);
+
+	// axis source provides the translation axis
+	TranslateGizmo->AxisSource = Cast<UObject>(AxisSource);
+
+	// parameter source maps axis-parameter-change to translation of TransformSource's transform
+	UGizmoPlaneTranslationParameterSource* ParamSource = UGizmoPlaneTranslationParameterSource::Construct(AxisSource, TransformSource, this);
+	ParamSource->PositionConstraintFunction = [this](const FVector& Pos, FVector& Snapped) { return PositionSnapFunction(Pos, Snapped); };
+	TranslateGizmo->ParameterSource = ParamSource;
+
+	// sub-component provides hit target
+	UGizmoComponentHitTarget* HitTarget = UGizmoComponentHitTarget::Construct(AxisComponent, this);
+	if (this->UpdateHoverFunction)
+	{
+		HitTarget->UpdateHoverFunction = [AxisComponent, this](bool bHovering) { this->UpdateHoverFunction(AxisComponent, bHovering); };
+	}
+	TranslateGizmo->HitTarget = HitTarget;
+
+	TranslateGizmo->StateTarget = Cast<UObject>(StateTargetIn);
+
+	TranslateGizmo->ShouldUseCustomDestinationFunc = [this]() { return ShouldAlignDestination(); };
+	TranslateGizmo->CustomDestinationFunc =
+		[this](const UPlanePositionGizmo::FCustomDestinationParams& Params, FVector& OutputPoint) {
+		return DestinationAlignmentRayCaster(*Params.WorldRay, OutputPoint);
+	};
+
+	ActiveGizmos.Add(TranslateGizmo);
+	return TranslateGizmo;
+}
+
+
+
+
+
+UInteractiveGizmo* UTransformGizmo::AddAxisRotationGizmo(
+	UPrimitiveComponent* AxisComponent, USceneComponent* RootComponent,
+	IGizmoAxisSource* AxisSource,
+	IGizmoTransformSource* TransformSource,
+	IGizmoStateTarget* StateTargetIn)
+{
+	// create axis-angle gizmo, angle will drive axis-rotation
+	UAxisAngleGizmo* RotateGizmo = Cast<UAxisAngleGizmo>(GetGizmoManager()->CreateGizmo(AxisAngleBuilderIdentifier));
+	check(RotateGizmo);
+
+	// axis source provides the rotation axis
+	RotateGizmo->AxisSource = Cast<UObject>(AxisSource);
+
+	// parameter source maps angle-parameter-change to rotation of TransformSource's transform
+	UGizmoAxisRotationParameterSource* AngleSource = UGizmoAxisRotationParameterSource::Construct(AxisSource, TransformSource, this);
+	AngleSource->RotationConstraintFunction = [this](const FQuat& DeltaRotation){ return RotationSnapFunction(DeltaRotation); };
+	RotateGizmo->AngleSource = AngleSource;
+
+	// sub-component provides hit target
+	UGizmoComponentHitTarget* HitTarget = UGizmoComponentHitTarget::Construct(AxisComponent, this);
+	if (this->UpdateHoverFunction)
+	{
+		HitTarget->UpdateHoverFunction = [AxisComponent, this](bool bHovering) { this->UpdateHoverFunction(AxisComponent, bHovering); };
+	}
+	RotateGizmo->HitTarget = HitTarget;
+
+	RotateGizmo->StateTarget = Cast<UObject>(StateTargetIn);
+
+	RotateGizmo->ShouldUseCustomDestinationFunc = [this]() { return ShouldAlignDestination(); };
+	RotateGizmo->CustomDestinationFunc =
+		[this](const UAxisAngleGizmo::FCustomDestinationParams& Params, FVector& OutputPoint) {
+		return DestinationAlignmentRayCaster(*Params.WorldRay, OutputPoint);
+	};
+
+	ActiveGizmos.Add(RotateGizmo);
+
+	return RotateGizmo;
+}
+
+
+
+UInteractiveGizmo* UTransformGizmo::AddAxisScaleGizmo(
+	UPrimitiveComponent* AxisComponent, USceneComponent* RootComponent,
+	IGizmoAxisSource* GizmoAxisSource, IGizmoAxisSource* ParameterAxisSource,
+	IGizmoTransformSource* TransformSource,
+	IGizmoStateTarget* StateTargetIn)
+{
+	// create axis-position gizmo, axis-position parameter will drive scale
+	UAxisPositionGizmo* ScaleGizmo = Cast<UAxisPositionGizmo>(GetGizmoManager()->CreateGizmo(AxisPositionBuilderIdentifier));
+	ScaleGizmo->bEnableSignedAxis = true;
+	check(ScaleGizmo);
+
+	// axis source provides the translation axis
+	ScaleGizmo->AxisSource = Cast<UObject>(GizmoAxisSource);
+
+	// parameter source maps axis-parameter-change to translation of TransformSource's transform
+	UGizmoAxisScaleParameterSource* ParamSource = UGizmoAxisScaleParameterSource::Construct(ParameterAxisSource, TransformSource, this);
+	ParamSource->bClampToZero = bDisallowNegativeScaling;
+	ScaleGizmo->ParameterSource = ParamSource;
+
+	// sub-component provides hit target
+	UGizmoComponentHitTarget* HitTarget = UGizmoComponentHitTarget::Construct(AxisComponent, this);
+	if (this->UpdateHoverFunction)
+	{
+		HitTarget->UpdateHoverFunction = [AxisComponent, this](bool bHovering) { this->UpdateHoverFunction(AxisComponent, bHovering); };
+	}
+	ScaleGizmo->HitTarget = HitTarget;
+
+	ScaleGizmo->StateTarget = Cast<UObject>(StateTargetIn);
+
+	ActiveGizmos.Add(ScaleGizmo);
+	return ScaleGizmo;
+}
+
+
+
+UInteractiveGizmo* UTransformGizmo::AddPlaneScaleGizmo(
+	UPrimitiveComponent* AxisComponent, USceneComponent* RootComponent,
+	IGizmoAxisSource* GizmoAxisSource, IGizmoAxisSource* ParameterAxisSource,
+	IGizmoTransformSource* TransformSource,
+	IGizmoStateTarget* StateTargetIn)
+{
+	// create axis-position gizmo, axis-position parameter will drive scale
+	UPlanePositionGizmo* ScaleGizmo = Cast<UPlanePositionGizmo>(GetGizmoManager()->CreateGizmo(PlanePositionBuilderIdentifier));
+	ScaleGizmo->bEnableSignedAxis = true;
+	check(ScaleGizmo);
+
+	// axis source provides the translation axis
+	ScaleGizmo->AxisSource = Cast<UObject>(GizmoAxisSource);
+
+	// parameter source maps axis-parameter-change to translation of TransformSource's transform
+	UGizmoPlaneScaleParameterSource* ParamSource = UGizmoPlaneScaleParameterSource::Construct(ParameterAxisSource, TransformSource, this);
+	ParamSource->bClampToZero = bDisallowNegativeScaling;
+	ParamSource->bUseEqualScaling = true;
+	ScaleGizmo->ParameterSource = ParamSource;
+
+	// sub-component provides hit target
+	UGizmoComponentHitTarget* HitTarget = UGizmoComponentHitTarget::Construct(AxisComponent, this);
+	if (this->UpdateHoverFunction)
+	{
+		HitTarget->UpdateHoverFunction = [AxisComponent, this](bool bHovering) { this->UpdateHoverFunction(AxisComponent, bHovering); };
+	}
+	ScaleGizmo->HitTarget = HitTarget;
+
+	ScaleGizmo->StateTarget = Cast<UObject>(StateTargetIn);
+
+	ActiveGizmos.Add(ScaleGizmo);
+	return ScaleGizmo;
+}
+
+
+
+
+
+UInteractiveGizmo* UTransformGizmo::AddUniformScaleGizmo(
+	UPrimitiveComponent* ScaleComponent, USceneComponent* RootComponent,
+	IGizmoAxisSource* GizmoAxisSource, IGizmoAxisSource* ParameterAxisSource,
+	IGizmoTransformSource* TransformSource,
+	IGizmoStateTarget* StateTargetIn)
+{
+	// create plane-position gizmo, plane-position parameter will drive scale
+	UPlanePositionGizmo* ScaleGizmo = Cast<UPlanePositionGizmo>(GetGizmoManager()->CreateGizmo(PlanePositionBuilderIdentifier));
+	check(ScaleGizmo);
+
+	// axis source provides the translation plane
+	ScaleGizmo->AxisSource = Cast<UObject>(GizmoAxisSource);
+
+	// parameter source maps axis-parameter-change to translation of TransformSource's transform
+	UGizmoUniformScaleParameterSource* ParamSource = UGizmoUniformScaleParameterSource::Construct(ParameterAxisSource, TransformSource, this);
+	//ParamSource->PositionConstraintFunction = [this](const FVector& Pos, FVector& Snapped) { return PositionSnapFunction(Pos, Snapped); };
+	ScaleGizmo->ParameterSource = ParamSource;
+
+	// sub-component provides hit target
+	UGizmoComponentHitTarget* HitTarget = UGizmoComponentHitTarget::Construct(ScaleComponent, this);
+	if (this->UpdateHoverFunction)
+	{
+		HitTarget->UpdateHoverFunction = [ScaleComponent, this](bool bHovering) { this->UpdateHoverFunction(ScaleComponent, bHovering); };
+	}
+	ScaleGizmo->HitTarget = HitTarget;
+
+	ScaleGizmo->StateTarget = Cast<UObject>(StateTargetIn);
+
+	ActiveGizmos.Add(ScaleGizmo);
+	return ScaleGizmo;
+}
+
+
+
+void UTransformGizmo::ClearActiveTarget()
+{
+	for (UInteractiveGizmo* Gizmo : ActiveGizmos)
+	{
+		GetGizmoManager()->DestroyGizmo(Gizmo);
+	}
+	ActiveGizmos.SetNum(0);
+	ActiveComponents.SetNum(0);
+	NonuniformScaleComponents.SetNum(0);
+
+	CameraAxisSource = nullptr;
+	AxisXSource = nullptr;
+	AxisYSource = nullptr;
+	AxisZSource = nullptr;
+	UnitAxisXSource = nullptr;
+	UnitAxisYSource = nullptr;
+	UnitAxisZSource = nullptr;
+	StateTarget = nullptr;
+
+	ActiveTarget = nullptr;
+}
+
+
+
+
+bool UTransformGizmo::PositionSnapFunction(const FVector& WorldPosition, FVector& SnappedPositionOut) const
+{
+	SnappedPositionOut = WorldPosition;
+
+	// only snap if we want snapping obvs
+	if (bSnapToWorldGrid == false)
+	{
+		return false;
+	}
+
+	// only snap to world grid when using world axes
+	EToolContextCoordinateSystem CoordSystem = GetGizmoManager()->GetContextQueriesAPI()->GetCurrentCoordinateSystem();
+	if (CoordSystem != EToolContextCoordinateSystem::World)
+	{
+		return false;
+	}
+
+	FSceneSnapQueryRequest Request;
+	Request.RequestType = ESceneSnapQueryType::Position;
+	Request.TargetTypes = ESceneSnapQueryTargetType::Grid;
+	Request.Position = WorldPosition;
+	if ( bGridSizeIsExplicit )
+	{
+		Request.GridSize = ExplicitGridSize;
+	}
+	TArray<FSceneSnapQueryResult> Results;
+	if (GetGizmoManager()->GetContextQueriesAPI()->ExecuteSceneSnapQuery(Request, Results))
+	{
+		SnappedPositionOut = Results[0].Position;
+		return true;
+	};
+
+	return false;
+}
+
+FQuat UTransformGizmo::RotationSnapFunction(const FQuat& DeltaRotation) const
+{
+	FQuat SnappedDeltaRotation = DeltaRotation;
+
+	// only snap if we want snapping obvs
+	if (bSnapToWorldRotGrid)
+	{
+		FSceneSnapQueryRequest Request;
+		Request.RequestType   = ESceneSnapQueryType::Rotation;
+		Request.TargetTypes   = ESceneSnapQueryTargetType::Grid;
+		Request.DeltaRotation = DeltaRotation;
+		if ( bRotationGridSizeIsExplicit )
+		{
+			Request.RotGridSize = ExplicitRotationGridSize;
+		}
+		TArray<FSceneSnapQueryResult> Results;
+		if (GetGizmoManager()->GetContextQueriesAPI()->ExecuteSceneSnapQuery(Request, Results))
+		{
+			SnappedDeltaRotation = Results[0].DeltaRotation;
+		};
+	}
+	return SnappedDeltaRotation;
+}
+
+#undef LOCTEXT_NAMESPACE
diff --git a/Engine/Source/Runtime/InteractiveToolsFramework/Public/BaseGizmos/GizmoBoxComponent.h b/Engine/Source/Runtime/InteractiveToolsFramework/Public/BaseGizmos/GizmoBoxComponent.h
new file mode 100644
--- /dev/null
+++ b/Engine/Source/Runtime/InteractiveToolsFramework/Public/BaseGizmos/GizmoBoxComponent.h
@@ -0,0 +1,49 @@
+// Copyright Epic Games, Inc. All Rights Reserved.
+
+#pragma once
+
+#include "CoreMinimal.h"
+#include "UObject/ObjectMacros.h"
+#include "BaseGizmos/GizmoBaseComponent.h"
+#include "GizmoBoxComponent.generated.h"
+
+/**
+ * Simple Component intended to be used as part of 3D Gizmos. 
+ * Draws outline of 3D Box based on parameters.
+ */
+UCLASS(ClassGroup = Utility, HideCategories = (Physics, Collision, Lighting, Rendering, Mobile))
+class INTERACTIVETOOLSFRAMEWORK_API UGizmoBoxComponent : public UGizmoBaseComponent
+{
+	GENERATED_BODY()
+
+public:
+	UPROPERTY(EditAnywhere, Category = Options)
+	FVector Origin = FVector(0, 0, 0);
+
+	UPROPERTY(EditAnywhere, Category = Options)
+	FQuat Rotation = FQuat::Identity;
+
+	UPROPERTY(EditAnywhere, Category = Options)
+	FVector Dimensions = FVector(20.0f, 20.0f, 20.0f);
+
+	UPROPERTY(EditAnywhere, Category = Options)
+	float LineThickness = 2.0f;
+
+	UPROPERTY(EditAnywhere, Category = Options)
+	bool bRemoveHiddenLines = true;
+
+	UPROPERTY(EditAnywhere, Category = Options)
+	bool bEnableAxisFlip = true;
+
+private:
+	//~ Begin UPrimitiveComponent Interface.
+	virtual FPrimitiveSceneProxy* CreateSceneProxy() override;
+	virtual bool LineTraceComponent(FHitResult& OutHit, const FVector Start, const FVector End, const FCollisionQueryParams& Params) override;
+	virtual void GetUsedMaterials(TArray<UMaterialInterface*>& OutMaterials, bool bGetDebugMaterials = false) const override;
+	//~ End UPrimitiveComponent Interface.
+
+	//~ Begin USceneComponent Interface.
+	virtual FBoxSphereBounds CalcBounds(const FTransform& LocalToWorld) const override;
+	//~ Begin USceneComponent Interface.
+
+};
\ No newline at end of file
diff --git a/Engine/Source/Runtime/InteractiveToolsFramework/Public/BaseGizmos/ParameterToTransformAdapters.h b/Engine/Source/Runtime/InteractiveToolsFramework/Public/BaseGizmos/ParameterToTransformAdapters.h
new file mode 100644
--- /dev/null
+++ b/Engine/Source/Runtime/InteractiveToolsFramework/Public/BaseGizmos/ParameterToTransformAdapters.h
@@ -0,0 +1,820 @@
+// Copyright Epic Games, Inc. All Rights Reserved.
+
+#pragma once
+
+#include "CoreMinimal.h"
+#include "BaseGizmos/GizmoInterfaces.h"
+#include "BaseGizmos/ParameterSourcesFloat.h"
+#include "BaseGizmos/ParameterSourcesVec2.h"
+#include "BaseGizmos/GizmoMath.h"
+#include "ParameterToTransformAdapters.generated.h"
+
+
+//
+// Various 1D and 2D ParameterSource converters intended to be used to create 3D transformation gizmos.
+// Based on base classes in ParameterSourcesFloat.h and ParameterSourcesVec2.h
+// 
+
+
+
+/**
+ * UGizmoAxisTranslationParameterSource is an IGizmoFloatParameterSource implementation that
+ * interprets the float value as the parameter of a line equation, and maps this parameter to a 3D translation 
+ * along a line with origin/direction given by an IGizmoAxisSource. This translation is applied to an IGizmoTransformSource.
+ *
+ * This ParameterSource is intended to be used to create 3D Axis Translation Gizmos.
+ */
+UCLASS()
+class INTERACTIVETOOLSFRAMEWORK_API UGizmoAxisTranslationParameterSource : public UGizmoBaseFloatParameterSource
+{
+	GENERATED_BODY()
+public:
+
+	/**
+	 * Optional position constraint function. Called during interaction with the new transform origin.
+	 * To snap the transform to a new position, return as second value, and return true from your lambda.
+	 * Note that returned snap point will be projected onto the current translation origin/axis.
+	 * @return true if constraint point was found and should be used, false to ignore
+	 */
+	TUniqueFunction<bool(const FVector& Input, FVector& PreProjectedOutput)> PositionConstraintFunction = [](const FVector&, FVector&) { return false; };
+
+	virtual float GetParameter() const override
+	{
+		return Parameter;
+	}
+
+	virtual void SetParameter(float NewValue) override
+	{
+		Parameter = NewValue;
+		LastChange.CurrentValue = NewValue;
+
+		// construct translation as delta from initial position
+		FVector Translation = LastChange.GetChangeDelta() * CurTranslationAxis;
+
+		// translate the initial transform
+		FTransform NewTransform = InitialTransform;
+		NewTransform.AddToTranslation(Translation);
+
+		// apply position constraint
+		FVector SnappedPos;
+		if (PositionConstraintFunction(NewTransform.GetTranslation(), SnappedPos))
+		{
+			FVector SnappedLinePos = GizmoMath::ProjectPointOntoLine(SnappedPos, CurTranslationOrigin, CurTranslationAxis);
+			NewTransform.SetTranslation(SnappedLinePos);
+		}
+
+		TransformSource->SetTransform(NewTransform);
+
+		OnParameterChanged.Broadcast(this, LastChange);
+	}
+
+	virtual void BeginModify()
+	{
+		check(AxisSource);
+
+		LastChange = FGizmoFloatParameterChange(Parameter);
+
+		InitialTransform = TransformSource->GetTransform();
+		CurTranslationAxis = AxisSource->GetDirection();
+		CurTranslationOrigin = AxisSource->GetOrigin();
+	}
+
+	virtual void EndModify()
+	{
+	}
+
+
+public:
+	/** The Parameter line-equation value is converted to a 3D Translation along this Axis */
+	UPROPERTY()
+	TScriptInterface<IGizmoAxisSource> AxisSource;
+
+	/** This TransformSource is updated by applying the constructed 3D translation */
+	UPROPERTY()
+	TScriptInterface<IGizmoTransformSource> TransformSource;
+
+
+public:
+	/** Parameter is the line-equation parameter that this FloatParameterSource provides */
+	UPROPERTY()
+	float Parameter = 0.0f;
+
+	/** Active parameter change (only valid between BeginModify/EndModify) */
+	UPROPERTY()
+	FGizmoFloatParameterChange LastChange;
+
+	/** tranlsation axis for current parameter edit (only valid between BeginModify/EndModify) */
+	UPROPERTY()
+	FVector CurTranslationAxis;
+
+	/** tranlsation origin for current parameter edit (only valid between BeginModify/EndModify) */
+	UPROPERTY()
+	FVector CurTranslationOrigin;
+
+	/** Saved copy of Initial Transform for current parameter edit (only valid between BeginModify/EndModify) */
+	UPROPERTY()
+	FTransform InitialTransform;
+
+public:
+
+	/**
+	 * Create a standard instance of this ParameterSource, with the given AxisSource and TransformSource
+	 */
+	static UGizmoAxisTranslationParameterSource* Construct(
+		IGizmoAxisSource* AxisSourceIn,
+		IGizmoTransformSource* TransformSourceIn,
+		UObject* Outer = (UObject*)GetTransientPackage())
+	{
+		UGizmoAxisTranslationParameterSource* NewSource = NewObject<UGizmoAxisTranslationParameterSource>(Outer);
+		NewSource->AxisSource = Cast<UObject>(AxisSourceIn);
+		NewSource->TransformSource = Cast<UObject>(TransformSourceIn);
+		return NewSource;
+	}
+};
+
+
+
+
+
+/**
+ * UGizmoAxisRotationParameterSource is an IGizmoVec2ParameterSource implementation that
+ * interprets the FVector2D parameter as a position in a 2D plane, and maps this position to a 3D translation
+ * a plane with origin/normal given by an IGizmoAxisSource. This translation is applied to an IGizmoTransformSource.
+ * 
+ * This ParameterSource is intended to be used to create 3D Plane Translation Gizmos.
+ */
+UCLASS()
+class INTERACTIVETOOLSFRAMEWORK_API UGizmoPlaneTranslationParameterSource : public UGizmoBaseVec2ParameterSource
+{
+	GENERATED_BODY()
+public:
+
+	/**
+	 * Optional position constraint function. Called during interaction with the new transform origin.
+	 * To snap the transform to a new position, return as second value, and return true from your lambda.
+	 * Note that returned snap point will be projected onto the current translation origin/axis.
+	 * @return true if constraint point was found and should be used, false to ignore
+	 */
+	TUniqueFunction<bool(const FVector&, FVector&)> PositionConstraintFunction = [](const FVector&, FVector&) { return false; };
+
+
+	virtual FVector2D GetParameter() const override
+	{
+		return Parameter;
+	}
+
+	virtual void SetParameter(const FVector2D& NewValue) override
+	{
+		Parameter = NewValue;
+		LastChange.CurrentValue = NewValue;
+
+		// construct translation as delta from initial position
+		FVector2D Delta = LastChange.GetChangeDelta();
+		FVector Translation = Delta.X*CurTranslationAxisX + Delta.Y*CurTranslationAxisY;
+
+		// apply translation to initial transform
+		FTransform NewTransform = InitialTransform;
+		NewTransform.AddToTranslation(Translation);
+
+		// apply position constraint
+		FVector SnappedPos;
+		if (PositionConstraintFunction(NewTransform.GetTranslation(), SnappedPos))
+		{
+			FVector PlanePos = GizmoMath::ProjectPointOntoPlane(SnappedPos, CurTranslationOrigin, CurTranslationNormal);
+			NewTransform.SetTranslation(PlanePos);
+		}
+
+		TransformSource->SetTransform(NewTransform);
+
+		OnParameterChanged.Broadcast(this, LastChange);
+	}
+
+	virtual void BeginModify()
+	{
+		check(AxisSource);
+
+		LastChange = FGizmoVec2ParameterChange(Parameter);
+
+		// save initial transformation and axis information
+		InitialTransform = TransformSource->GetTransform();
+		CurTranslationOrigin = AxisSource->GetOrigin();
+		AxisSource->GetAxisFrame(CurTranslationNormal, CurTranslationAxisX, CurTranslationAxisY);
+	}
+
+	virtual void EndModify()
+	{
+	}
+
+
+public:
+	/** AxisSource provides the 3D plane (origin/normal/u/v) that is used to interpret the 2D parameters */
+	UPROPERTY()
+	TScriptInterface<IGizmoAxisSource> AxisSource;
+
+	/** This TransformSource is updated by applying the constructed 3D translation */
+	UPROPERTY()
+	TScriptInterface<IGizmoTransformSource> TransformSource;
+
+
+public:
+	/** Parameter is the two line-equation parameters that this Vec2ParameterSource provides */
+	UPROPERTY()
+	FVector2D Parameter = FVector2D::ZeroVector;
+
+	/** Active parameter change (only valid between BeginModify/EndModify) */
+	UPROPERTY()
+	FGizmoVec2ParameterChange LastChange;
+
+	/** plane origin for current parameter edit (only valid between BeginModify/EndModify) */
+	UPROPERTY()
+	FVector CurTranslationOrigin;
+
+	/** plane normal for current parameter edit (only valid between BeginModify/EndModify) */
+	UPROPERTY()
+	FVector CurTranslationNormal;
+
+	/** in-plane axis X for current parameter edit (only valid between BeginModify/EndModify) */
+	UPROPERTY()
+	FVector CurTranslationAxisX;
+
+	/** in-plane axis Y for current parameter edit (only valid between BeginModify/EndModify) */
+	UPROPERTY()
+	FVector CurTranslationAxisY;
+
+	/** Saved copy of Initial Transform for current parameter edit (only valid between BeginModify/EndModify) */
+	UPROPERTY()
+	FTransform InitialTransform;
+
+public:
+
+	/**
+	 * Create a standard instance of this ParameterSource, with the given AxisSource and TransformSource
+	 */
+	static UGizmoPlaneTranslationParameterSource* Construct(
+		IGizmoAxisSource* AxisSourceIn,
+		IGizmoTransformSource* TransformSourceIn,
+		UObject* Outer = (UObject*)GetTransientPackage())
+	{
+		UGizmoPlaneTranslationParameterSource* NewSource = NewObject<UGizmoPlaneTranslationParameterSource>(Outer);
+		NewSource->AxisSource = Cast<UObject>(AxisSourceIn);
+		NewSource->TransformSource = Cast<UObject>(TransformSourceIn);
+		return NewSource;
+	}
+};
+
+
+
+
+
+
+/**
+ * UGizmoAxisRotationParameterSource is an IGizmoFloatParameterSource implementation that
+ * interprets the float parameter as an angle, and maps this angle to a 3D rotation
+ * around an IGizmoAxisSource (ie 3D axis). This rotation is applied to an IGizmoTransformSource.
+ * This ParameterSource is intended to be used to create 3D Rotation Gizmos.
+ */
+UCLASS()
+class INTERACTIVETOOLSFRAMEWORK_API UGizmoAxisRotationParameterSource : public UGizmoBaseFloatParameterSource
+{
+	GENERATED_BODY()
+public:
+	/**
+	 * Optional rotation constraint function. Called during interaction with the rotation delta
+	 * To snap the rotation delta, return the snapped quat
+	 * @return The snapped value of the rotation delta
+	 */
+	TUniqueFunction<FQuat(const FQuat&)> RotationConstraintFunction = [](const FQuat& DeltaRotation) { return DeltaRotation; };
+
+	virtual float GetParameter() const override
+	{
+		return Angle;
+	}
+
+	virtual void SetParameter(float NewValue) override
+	{
+		Angle = NewValue;
+		LastChange.CurrentValue = NewValue;
+
+		// construct rotation as delta from initial position
+		FQuat DeltaRotation(CurRotationAxis, LastChange.GetChangeDelta());
+		DeltaRotation = RotationConstraintFunction(DeltaRotation);
+
+		// rotate the vector from the rotation origin to the transform origin, 
+		// to get the translation of the origin produced by the rotation
+		FVector DeltaPosition = InitialTransform.GetLocation() - CurRotationOrigin;
+		DeltaPosition = DeltaRotation * DeltaPosition;
+		FVector NewLocation = CurRotationOrigin + DeltaPosition;
+
+		// rotate the initial transform by the rotation
+		FQuat NewRotation = DeltaRotation * InitialTransform.GetRotation();
+
+		// construct new transform
+		FTransform NewTransform = InitialTransform;
+		NewTransform.SetLocation(NewLocation);
+		NewTransform.SetRotation(NewRotation);
+		TransformSource->SetTransform(NewTransform);
+
+		OnParameterChanged.Broadcast(this, LastChange);
+	}
+
+	virtual void BeginModify()
+	{
+		check(AxisSource != nullptr);
+
+		LastChange = FGizmoFloatParameterChange(Angle);
+
+		// save initial transformation and axis information
+		InitialTransform = TransformSource->GetTransform();
+		CurRotationAxis = AxisSource->GetDirection();
+		CurRotationOrigin = AxisSource->GetOrigin();
+	}
+
+	virtual void EndModify()
+	{
+	}
+
+
+public:
+	/** float-parameter Angle is mapped to a 3D Rotation around this Axis */
+	UPROPERTY()
+	TScriptInterface<IGizmoAxisSource> AxisSource;
+
+	/** This TransformSource is updated by applying the constructed 3D rotation */
+	UPROPERTY()
+	TScriptInterface<IGizmoTransformSource> TransformSource;
+
+
+public:
+	/** Angle is the parameter that this FloatParameterSource provides */
+	UPROPERTY()
+	float Angle = 0.0f;
+
+	/** Active parameter change (only valid between BeginModify/EndModify) */
+	UPROPERTY()
+	FGizmoFloatParameterChange LastChange;
+
+	/** Rotation axis for current parameter edit (only valid between BeginModify/EndModify) */
+	UPROPERTY()
+	FVector CurRotationAxis;
+
+	/** Rotation origin for current parameter edit (only valid between BeginModify/EndModify) */
+	UPROPERTY()
+	FVector CurRotationOrigin;
+
+	/** Saved copy of Initial Transform for current parameter edit (only valid between BeginModify/EndModify) */
+	UPROPERTY()
+	FTransform InitialTransform;
+
+
+public:
+
+	/**
+	 * Create a standard instance of this ParameterSource, with the given AxisSource and TransformSource
+	 */
+	static UGizmoAxisRotationParameterSource* Construct(
+		IGizmoAxisSource* AxisSourceIn,
+		IGizmoTransformSource* TransformSourceIn,
+		UObject* Outer = (UObject*)GetTransientPackage())
+	{
+		UGizmoAxisRotationParameterSource* NewSource = NewObject<UGizmoAxisRotationParameterSource>(Outer);
+		NewSource->AxisSource = Cast<UObject>(AxisSourceIn);
+		NewSource->TransformSource = Cast<UObject>(TransformSourceIn);
+		return NewSource;
+	}
+};
+
+
+
+
+
+
+
+
+
+ /**
+  * UGizmoUniformScaleParameterSource is an IGizmoVec2ParameterSource implementation that
+  * interprets the Vec2 parameter as the position in a 2D plane, and maps this parameter to a
+  * change in a uniform scale value. This scale is applied to an IGizmoTransformSource.
+  * 
+  * The X/Y plane parameter deltas are converted to a uniform scale delta simply by adding them.
+  * This assumpes that the 3D plane is oriented such that "positive" along the X and Y 
+  * tangent axes corresponds to something semantically meaningful. It's up to the client to do this.
+  *
+  * This ParameterSource is intended to be used to create 3D Uniform Scale Gizmos.
+  */
+UCLASS()
+class INTERACTIVETOOLSFRAMEWORK_API UGizmoUniformScaleParameterSource : public UGizmoBaseVec2ParameterSource
+{
+	GENERATED_BODY()
+public:
+
+	/**
+	 * Optional position constraint function. Called during interaction with the new transform origin.
+	 * To snap the transform to a new position, return as second value, and return true from your lambda.
+	 * Note that returned snap point will be projected onto the current translation origin/axis.
+	 * @return true if constraint point was found and should be used, false to ignore
+	 */
+	//TUniqueFunction<bool(const FVector&, FVector&)> PositionConstraintFunction = [](const FVector&, FVector&) { return false; };
+
+
+	virtual FVector2D GetParameter() const override
+	{
+		return Parameter;
+	}
+
+	virtual void SetParameter(const FVector2D& NewValue) override
+	{
+		Parameter = NewValue;
+		LastChange.CurrentValue = NewValue;
+
+		// Convert 2D parameter delta to a 1D uniform scale change
+		// This possibly be exposed as a TFunction to allow customization?
+		double SignedDelta = LastChange.GetChangeDelta().X + LastChange.GetChangeDelta().Y;
+		SignedDelta *= ScaleMultiplier;
+		SignedDelta += 1.0;
+
+		FTransform NewTransform = InitialTransform;
+		FVector StartScale = InitialTransform.GetScale3D();
+		FVector NewScale = SignedDelta * StartScale;
+		NewTransform.SetScale3D(NewScale);
+
+		// apply position constraint
+		//FVector SnappedPos;
+		//if (PositionConstraintFunction(NewTransform.GetScale(), SnappedPos))
+		//{
+		//	FVector PlanePos = GizmoMath::ProjectPointOntoPlane(SnappedPos, CurTranslationOrigin, CurTranslationNormal);
+		//	NewTransform.SetTranslation(PlanePos);
+		//}
+
+		TransformSource->SetTransform(NewTransform);
+
+		OnParameterChanged.Broadcast(this, LastChange);
+	}
+
+	virtual void BeginModify()
+	{
+		check(AxisSource);
+
+		LastChange = FGizmoVec2ParameterChange(Parameter);
+
+		// save initial transformation and axis information
+		InitialTransform = TransformSource->GetTransform();
+		CurScaleOrigin = AxisSource->GetOrigin();
+		// note: currently not used!
+		AxisSource->GetAxisFrame(CurScaleNormal, CurScaleAxisX, CurScaleAxisY);
+	}
+
+	virtual void EndModify()
+	{
+	}
+
+
+public:
+	/** AxisSource provides the 3D plane (origin/normal/u/v) that is used to interpret the 2D parameters */
+	UPROPERTY()
+	TScriptInterface<IGizmoAxisSource> AxisSource;
+
+	/** This TransformSource is updated by applying the constructed 3D translation */
+	UPROPERTY()
+	TScriptInterface<IGizmoTransformSource> TransformSource;
+
+	/** Coordinate delta is multiplied by this amount */
+	UPROPERTY()
+	float ScaleMultiplier = 0.025f;
+
+public:
+	/** Parameter is the two line-equation parameters that this Vec2ParameterSource provides */
+	UPROPERTY()
+	FVector2D Parameter = FVector2D::ZeroVector;
+
+	/** Active parameter change (only valid between BeginModify/EndModify) */
+	UPROPERTY()
+	FGizmoVec2ParameterChange LastChange;
+
+	/** plane origin for current parameter edit (only valid between BeginModify/EndModify) */
+	UPROPERTY()
+	FVector CurScaleOrigin;
+
+	/** plane normal for current parameter edit (only valid between BeginModify/EndModify) */
+	UPROPERTY()
+	FVector CurScaleNormal;
+
+	/** in-plane axis X for current parameter edit (only valid between BeginModify/EndModify) */
+	UPROPERTY()
+	FVector CurScaleAxisX;
+
+	/** in-plane axis Y for current parameter edit (only valid between BeginModify/EndModify) */
+	UPROPERTY()
+	FVector CurScaleAxisY;
+
+	/** Saved copy of Initial Transform for current parameter edit (only valid between BeginModify/EndModify) */
+	UPROPERTY()
+	FTransform InitialTransform;
+
+public:
+
+	/**
+	 * Create a standard instance of this ParameterSource, with the given AxisSource and TransformSource
+	 */
+	static UGizmoUniformScaleParameterSource* Construct(
+		IGizmoAxisSource* AxisSourceIn,
+		IGizmoTransformSource* TransformSourceIn,
+		UObject* Outer = (UObject*)GetTransientPackage())
+	{
+		UGizmoUniformScaleParameterSource* NewSource = NewObject<UGizmoUniformScaleParameterSource>(Outer);
+		NewSource->AxisSource = Cast<UObject>(AxisSourceIn);
+		NewSource->TransformSource = Cast<UObject>(TransformSourceIn);
+		return NewSource;
+	}
+};
+
+
+
+
+
+/**
+ * UGizmoAxisScaleParameterSource is an IGizmoFloatParameterSource implementation that
+ * interprets the float value as the parameter of a line equation, and maps this parameter to a scale factor
+ * along a line with origin/direction given by an IGizmoAxisSource. This scale is applied to an IGizmoTransformSource.
+ *
+ * This ParameterSource is intended to be used to create 3D Axis Scale Gizmos.
+ */
+UCLASS()
+class INTERACTIVETOOLSFRAMEWORK_API UGizmoAxisScaleParameterSource : public UGizmoBaseFloatParameterSource
+{
+	GENERATED_BODY()
+public:
+
+	/**
+	 * Optional position constraint function. Called during interaction with the new transform origin.
+	 * To snap the transform to a new position, return as second value, and return true from your lambda.
+	 * Note that returned snap point will be projected onto the current translation origin/axis.
+	 * @return true if constraint point was found and should be used, false to ignore
+	 */
+	//TUniqueFunction<bool(const float&, float&)> ScaleConstraintFunction = [](const float&, float&) { return false; };
+
+	virtual float GetParameter() const override
+	{
+		return Parameter;
+	}
+
+	virtual void SetParameter(float NewValue) override
+	{
+		Parameter = NewValue;
+		LastChange.CurrentValue = NewValue;
+
+		// construct translation as delta from initial position
+		float ScaleDelta = LastChange.GetChangeDelta() * ScaleMultiplier;
+
+		// translate the initial transform
+		FTransform NewTransform = InitialTransform;
+		FVector StartScale = InitialTransform.GetScale3D();
+		FVector NewScale = StartScale + ScaleDelta * CurScaleAxis;
+
+		if (bClampToZero)
+		{
+			NewScale = FVector::Max(FVector::ZeroVector, NewScale);
+		}
+
+		NewTransform.SetScale3D(NewScale);
+
+
+		// apply position constraint
+		//FVector SnappedPos;
+		//if (ScaleConstraintFunction(NewTransform.GetTranslation(), SnappedPos))
+		//{
+		//	FVector SnappedLinePos = GizmoMath::ProjectPointOntoLine(SnappedPos, CurTranslationOrigin, CurTranslationAxis);
+		//	NewTransform.SetTranslation(SnappedLinePos);
+		//}
+
+		TransformSource->SetTransform(NewTransform);
+
+		OnParameterChanged.Broadcast(this, LastChange);
+	}
+
+	virtual void BeginModify()
+	{
+		check(AxisSource);
+
+		LastChange = FGizmoFloatParameterChange(Parameter);
+
+		InitialTransform = TransformSource->GetTransform();
+		
+		CurScaleAxis = AxisSource->GetDirection();
+		CurScaleOrigin = AxisSource->GetOrigin();
+	}
+
+	virtual void EndModify()
+	{
+	}
+
+
+public:
+	/** The Parameter line-equation value is converted to a 3D Translation along this Axis */
+	UPROPERTY()
+	TScriptInterface<IGizmoAxisSource> AxisSource;
+
+	/** This TransformSource is updated by applying the constructed 3D rotation */
+	UPROPERTY()
+	TScriptInterface<IGizmoTransformSource> TransformSource;
+
+	/** Coordinate delta is multiplied by this amount */
+	UPROPERTY()
+	float ScaleMultiplier = 0.02f;
+
+	/** If true, the minimal output scale will be zero. */
+	UPROPERTY()
+	bool bClampToZero = false;
+
+public:
+	/** Parameter is the line-equation parameter that this FloatParameterSource provides */
+	UPROPERTY()
+	float Parameter = 0.0f;
+
+	/** Active parameter change (only valid between BeginModify/EndModify) */
+	UPROPERTY()
+	FGizmoFloatParameterChange LastChange;
+
+	/** scale axis for current parameter edit (only valid between BeginModify/EndModify) */
+	UPROPERTY()
+	FVector CurScaleAxis;
+
+	/** scale origin for current parameter edit (only valid between BeginModify/EndModify) */
+	UPROPERTY()
+	FVector CurScaleOrigin;
+
+	/** Saved copy of Initial transform for current parameter edit (only valid between BeginModify/EndModify) */
+	UPROPERTY()
+	FTransform InitialTransform;
+
+public:
+
+	/**
+	 * Create a standard instance of this ParameterSource, with the given AxisSource and TransformSource
+	 */
+	static UGizmoAxisScaleParameterSource* Construct(
+		IGizmoAxisSource* AxisSourceIn,
+		IGizmoTransformSource* TransformSourceIn,
+		UObject* Outer = (UObject*)GetTransientPackage())
+	{
+		UGizmoAxisScaleParameterSource* NewSource = NewObject<UGizmoAxisScaleParameterSource>(Outer);
+		NewSource->AxisSource = Cast<UObject>(AxisSourceIn);
+		NewSource->TransformSource = Cast<UObject>(TransformSourceIn);
+		return NewSource;
+	}
+};
+
+
+
+
+
+
+
+
+/**
+ * UGizmoPlaneScaleParameterSource  is an UGizmoBaseVec2ParameterSource implementation that
+ * maps a 2D parameter delta to a change in a 3D scaling vector, based on the tangent axes of a 3D plane
+ * This scale is applied to an IGizmoTransformSource.
+ *
+ * This ParameterSource is intended to be used to create 3D Plane Scale Gizmos.
+ */
+UCLASS()
+class INTERACTIVETOOLSFRAMEWORK_API UGizmoPlaneScaleParameterSource : public UGizmoBaseVec2ParameterSource
+{
+	GENERATED_BODY()
+public:
+
+	/**
+	 * Optional position constraint function. Called during interaction with the new transform origin.
+	 * To snap the transform to a new position, return as second value, and return true from your lambda.
+	 * Note that returned snap point will be projected onto the current translation origin/axis.
+	 * @return true if constraint point was found and should be used, false to ignore
+	 */
+	TUniqueFunction<bool(const FVector&, FVector&)> PositionConstraintFunction = [](const FVector&, FVector&) { return false; };
+
+
+	virtual FVector2D GetParameter() const override
+	{
+		return Parameter;
+	}
+
+	virtual void SetParameter(const FVector2D& NewValue) override
+	{
+		Parameter = NewValue;
+		LastChange.CurrentValue = NewValue;
+
+		// construct Scale as delta from initial position
+		FVector2D Delta = LastChange.GetChangeDelta() * ScaleMultiplier;
+
+		if (bUseEqualScaling)
+		{
+			Delta = FVector2D(FMath::Min(Delta.X, Delta.Y));
+		}
+
+		FTransform NewTransform = InitialTransform;
+		FVector StartScale = InitialTransform.GetScale3D();
+		FVector NewScale = StartScale + Delta.X*CurScaleAxisX + Delta.Y*CurScaleAxisY;
+
+		if (bClampToZero)
+		{
+			NewScale = FVector::Max(NewScale, FVector::ZeroVector);
+		}
+
+		NewTransform.SetScale3D(NewScale);
+
+		// apply position constraint
+		//FVector SnappedPos;
+		//if (PositionConstraintFunction(NewTransform.GetScale(), SnappedPos))
+		//{
+		//	FVector PlanePos = GizmoMath::ProjectPointOntoPlane(SnappedPos, CurTranslationOrigin, CurTranslationNormal);
+		//	NewTransform.SetTranslation(PlanePos);
+		//}
+
+		TransformSource->SetTransform(NewTransform);
+
+		OnParameterChanged.Broadcast(this, LastChange);
+	}
+
+	virtual void BeginModify()
+	{
+		check(AxisSource);
+
+		LastChange = FGizmoVec2ParameterChange(Parameter);
+
+		// save initial transformation and axis information
+		InitialTransform = TransformSource->GetTransform();
+		CurScaleOrigin = AxisSource->GetOrigin();
+		AxisSource->GetAxisFrame(CurScaleNormal, CurScaleAxisX, CurScaleAxisY);
+	}
+
+	virtual void EndModify()
+	{
+	}
+
+
+public:
+	/** AxisSource provides the 3D plane (origin/normal/u/v) that is used to interpret the 2D parameters */
+	UPROPERTY()
+	TScriptInterface<IGizmoAxisSource> AxisSource;
+
+	/** This TransformSource is updated by applying the constructed 3D translation */
+	UPROPERTY()
+	TScriptInterface<IGizmoTransformSource> TransformSource;
+
+	/** Coordinate delta is multiplied by this amount */
+	UPROPERTY()
+	float ScaleMultiplier = 0.02f;
+
+	/** If true, the scaling will be done an equal amount in each axis, using the minimal value */
+	UPROPERTY()
+	bool bUseEqualScaling = false;
+
+	/** If true, the minimal output scale will be zero. */
+	UPROPERTY()
+	bool bClampToZero = false;
+public:
+	/** Parameter is the two line-equation parameters that this Vec2ParameterSource provides */
+	UPROPERTY()
+	FVector2D Parameter = FVector2D::ZeroVector;
+
+	/** Active parameter change (only valid between BeginModify/EndModify) */
+	UPROPERTY()
+	FGizmoVec2ParameterChange LastChange;
+
+	/** plane origin for current parameter edit (only valid between BeginModify/EndModify) */
+	UPROPERTY()
+	FVector CurScaleOrigin;
+
+	/** plane normal for current parameter edit (only valid between BeginModify/EndModify) */
+	UPROPERTY()
+	FVector CurScaleNormal;
+
+	/** in-plane axis X for current parameter edit (only valid between BeginModify/EndModify) */
+	UPROPERTY()
+	FVector CurScaleAxisX;
+
+	/** in-plane axis Y for current parameter edit (only valid between BeginModify/EndModify) */
+	UPROPERTY()
+	FVector CurScaleAxisY;
+
+	/** Saved copy of Initial Transform for current parameter edit (only valid between BeginModify/EndModify) */
+	UPROPERTY()
+	FTransform InitialTransform;
+
+public:
+
+	/**
+	 * Create a standard instance of this ParameterSource, with the given AxisSource and TransformSource
+	 */
+	static UGizmoPlaneScaleParameterSource* Construct(
+		IGizmoAxisSource* AxisSourceIn,
+		IGizmoTransformSource* TransformSourceIn,
+		UObject* Outer = (UObject*)GetTransientPackage())
+	{
+		UGizmoPlaneScaleParameterSource* NewSource = NewObject<UGizmoPlaneScaleParameterSource>(Outer);
+		NewSource->AxisSource = Cast<UObject>(AxisSourceIn);
+		NewSource->TransformSource = Cast<UObject>(TransformSourceIn);
+		return NewSource;
+	}
+};
+
diff --git a/Engine/Source/Runtime/InteractiveToolsFramework/Public/BaseTools/BaseBrushTool.h b/Engine/Source/Runtime/InteractiveToolsFramework/Public/BaseTools/BaseBrushTool.h
new file mode 100644
--- /dev/null
+++ b/Engine/Source/Runtime/InteractiveToolsFramework/Public/BaseTools/BaseBrushTool.h
@@ -0,0 +1,168 @@
+// Copyright Epic Games, Inc. All Rights Reserved.
+
+#pragma once
+
+#include "CoreMinimal.h"
+#include "BaseTools/MeshSurfacePointTool.h"
+#include "InteractiveTool.h"
+#include "BaseBrushTool.generated.h"
+
+class UBrushStampIndicator;
+/**
+ * Standard properties for a Brush-type Tool
+ */
+UCLASS()
+class INTERACTIVETOOLSFRAMEWORK_API UBrushBaseProperties : public UInteractiveToolPropertySet
+{
+	GENERATED_BODY()
+
+public:
+	UBrushBaseProperties();
+
+	/** Relative size of brush */
+	UPROPERTY(EditAnywhere, Category = Brush, meta = (DisplayName = "Size", UIMin = "0.0", UIMax = "1.0", ClampMin = "0.0", ClampMax = "10.0", DisplayPriority = 1, HideEditConditionToggle, EditCondition = "bSpecifyRadius == false"))
+	float BrushSize;
+
+	/** If true, ignore relative Brush Size and use explicit world Radius */
+	UPROPERTY(EditAnywhere, Category = Brush, AdvancedDisplay)
+	bool bSpecifyRadius;
+
+	/** Radius of brush */
+	UPROPERTY(EditAnywhere, Category = Brush, AdvancedDisplay, meta = (EditCondition = "bSpecifyRadius",
+		DisplayName = "Radius", UIMin = "1.0", UIMax = "1000.0", ClampMin = "0.1", ClampMax = "50000.0"))
+	float BrushRadius;
+
+	/** Strength of the brush (0.0 - 1.0) */
+	UPROPERTY(EditAnywhere, Category = Brush, meta = (DisplayName = "Strength", UIMin = "0.0", UIMax = "1.0", ClampMin = "0.0", ClampMax = "1.0", HideEditConditionToggle, EditConditionHides, EditCondition = "bShowStrength", DisplayPriority = 2))
+	float BrushStrength;
+
+	/** Amount of falloff to apply (0.0 - 1.0) */
+	UPROPERTY(EditAnywhere, Category = Brush, meta = (DisplayName = "Falloff", UIMin = "0.0", UIMax = "1.0", ClampMin = "0.0", ClampMax = "1.0", HideEditConditionToggle, EditConditionHides, EditCondition="bShowFalloff", DisplayPriority = 3))
+	float BrushFalloffAmount;
+
+	/** If false, then BrushStrength will not be shown in DetailsView panels (otherwise no effect) */
+	UPROPERTY( meta = (TransientToolProperty))
+	bool bShowStrength = true;
+
+	/** If false, then BrushFalloffAmount will not be shown in DetailsView panels (otherwise no effect) */
+	UPROPERTY(meta = (TransientToolProperty))
+	bool bShowFalloff = true;
+};
+
+
+/**
+ * Generic Brush Stamp data
+ */
+USTRUCT()
+struct INTERACTIVETOOLSFRAMEWORK_API FBrushStampData
+{
+	GENERATED_BODY();
+	/** Radius of brush stamp */
+	float Radius;
+	/** World Position of brush stamp */
+	FVector WorldPosition;
+	/** World Normal of brush stamp */
+	FVector WorldNormal;
+
+	/** Hit Result provided by implementations - may not be fully populated */
+	FHitResult HitResult;
+
+	/** Falloff of brush stamp */
+	float Falloff;
+};
+
+
+
+/**
+ * UBaseBrushTool implements standard brush-style functionality for an InteractiveTool.
+ * This includes:
+ *   1) brush radius property set with dimension-relative brush sizing and default brush radius hotkeys
+ *   2) brush indicator visualization
+ *   3) tracking of last brush stamp location via .LastBrushStamp FProperty
+ *   4) status of brush stroke via .bInBrushStroke FProperty
+ */
+UCLASS()
+class INTERACTIVETOOLSFRAMEWORK_API UBaseBrushTool : public UMeshSurfacePointTool
+{
+	GENERATED_BODY()
+
+public:
+	UBaseBrushTool();
+
+	virtual void RegisterActions(FInteractiveToolActionSet& ActionSet) override;
+
+	virtual void Setup() override;
+	virtual void Shutdown(EToolShutdownType ShutdownType) override;
+
+	virtual void Render(IToolsContextRenderAPI* RenderAPI) override;
+
+	virtual void OnPropertyModified(UObject* PropertySet, FProperty* Property) override;
+
+
+	//
+	// UMeshSurfacePointTool implementation
+	//
+	virtual void OnBeginDrag(const FRay& Ray) override;
+	virtual void OnUpdateDrag(const FRay& Ray) override;
+	virtual void OnEndDrag(const FRay& Ray) override;
+	virtual bool OnUpdateHover(const FInputDeviceRay& DevicePos) override;
+
+public:
+
+	/** Properties that control brush size/etc*/
+	UPROPERTY()
+	TObjectPtr<UBrushBaseProperties> BrushProperties;
+
+	/** Set to true by Tool if user is currently in an active brush stroke*/
+	UPROPERTY()
+	bool bInBrushStroke = false;
+
+	/** Uniform scale factor that scales from world space (where brush usually exists) to local space */
+	UPROPERTY()
+	float WorldToLocalScale = 1.0f;
+
+	/** Position of brush at last update (both during stroke and during Hover) */
+	UPROPERTY()
+	FBrushStampData LastBrushStamp;
+
+public:
+
+	virtual void IncreaseBrushSizeAction();
+	virtual void DecreaseBrushSizeAction();
+	virtual void IncreaseBrushStrengthAction();
+	virtual void DecreaseBrushStrengthAction();
+	virtual void IncreaseBrushFalloffAction();
+	virtual void DecreaseBrushFalloffAction();
+
+	virtual bool IsInBrushStroke() const { return bInBrushStroke; }
+
+	virtual double GetCurrentBrushRadius() const { return CurrentBrushRadius; }
+	virtual double GetCurrentBrushRadiusLocal() const { return CurrentBrushRadius * WorldToLocalScale; }
+
+protected:
+
+	/**
+	 * Subclasses should implement this to give an estimate of target dimension for brush size scaling
+	 */
+	virtual double EstimateMaximumTargetDimension() { return 100.0; }
+
+protected:
+	TInterval<float> BrushRelativeSizeRange;
+	double CurrentBrushRadius;
+	void RecalculateBrushRadius();
+
+	UPROPERTY()
+	TSoftClassPtr<UBrushBaseProperties> PropertyClass;
+
+	//
+	// Brush Indicator support
+	//
+protected:
+
+	UPROPERTY()
+	TObjectPtr<UBrushStampIndicator> BrushStampIndicator;
+
+	virtual void SetupBrushStampIndicator();
+	virtual void UpdateBrushStampIndicator();
+	virtual void ShutdownBrushStampIndicator();
+};
diff --git a/Engine/Source/Runtime/InteractiveToolsFramework/Public/BaseTools/MeshSurfacePointTool.h b/Engine/Source/Runtime/InteractiveToolsFramework/Public/BaseTools/MeshSurfacePointTool.h
new file mode 100644
--- /dev/null
+++ b/Engine/Source/Runtime/InteractiveToolsFramework/Public/BaseTools/MeshSurfacePointTool.h
@@ -0,0 +1,148 @@
+// Copyright Epic Games, Inc. All Rights Reserved.
+
+#pragma once
+
+#include "CoreMinimal.h"
+#include "SingleSelectionTool.h"
+#include "InteractiveToolBuilder.h"
+#include "InputBehaviorSet.h"
+#include "BaseBehaviors/BehaviorTargetInterfaces.h"
+#include "BaseBehaviors/AnyButtonInputBehavior.h"
+#include "MeshSurfacePointTool.generated.h"
+
+
+
+class UMeshSurfacePointTool;
+
+// This is a temporary interface to provide stylus pressure, currently necessary due to limitations
+// in the stylus plugin architecture. Should be removed once InputState/InputBehavior can support stylus events
+class INTERACTIVETOOLSFRAMEWORK_API IToolStylusStateProviderAPI
+{
+public:
+	virtual float GetCurrentPressure() const = 0;
+};
+
+/**
+ * 
+ */
+UCLASS()
+class INTERACTIVETOOLSFRAMEWORK_API UMeshSurfacePointToolBuilder : public UInteractiveToolBuilder
+{
+	GENERATED_BODY()
+
+public:
+	IToolStylusStateProviderAPI* StylusAPI = nullptr;
+
+	/** @return true if a single mesh source can be found in the active selection */
+	virtual bool CanBuildTool(const FToolBuilderState& SceneState) const override;
+
+	/** @return new Tool instance initialized with selected mesh source */
+	virtual UInteractiveTool* BuildTool(const FToolBuilderState& SceneState) const override;
+
+	/** @return new Tool instance. Override this in subclasses to build a different Tool class type */
+	virtual UMeshSurfacePointTool* CreateNewTool(const FToolBuilderState& SceneState) const;
+
+	/** Called by BuildTool to configure the Tool with the input MeshSource based on the SceneState */
+	virtual void InitializeNewTool(UMeshSurfacePointTool* Tool, const FToolBuilderState& SceneState) const;
+
+protected:
+	virtual const FToolTargetTypeRequirements& GetTargetRequirements() const override;
+};
+
+
+
+/**
+ * UMeshSurfacePointTool is a base Tool implementation that can be used to implement various
+ * "point on surface" interactions. The tool acts on an input IMeshDescriptionSource object,
+ * which the standard Builder can extract from the current selection (eg Editor selection).
+ * 
+ * Subclasses override the OnBeginDrag/OnUpdateDrag/OnEndDrag and OnUpdateHover functions
+ * to implement custom behavior.
+ */
+UCLASS()
+class INTERACTIVETOOLSFRAMEWORK_API UMeshSurfacePointTool : public USingleSelectionTool, public IClickDragBehaviorTarget, public IHoverBehaviorTarget
+{
+	GENERATED_BODY()
+
+public:
+	// UInteractiveTool API Implementation
+
+	/** Register InputBehaviors, etc */
+	virtual void Setup() override;
+
+	/** Set current stlyus API source */
+	virtual void SetStylusAPI(IToolStylusStateProviderAPI* StylusAPI);
+
+	// UMeshSurfacePointTool API
+
+	/**
+	 * @return true if the target MeshSource is hit by the Ray
+	 */
+	virtual bool HitTest(const FRay& Ray, FHitResult& OutHit);
+
+	/**
+	 * This function is called when the user begins a click-drag-release interaction
+	 */
+	virtual void OnBeginDrag(const FRay& Ray);
+
+	/**
+	 * This function is called each frame that the user is in a click-drag-release interaction
+	 */
+	virtual void OnUpdateDrag(const FRay& Ray);
+
+	/**
+	 * This function is called when the user releases the button driving a click-drag-release interaction
+	 */
+	virtual void OnEndDrag(const FRay& Ray);
+
+
+
+
+	// IClickDragBehaviorTarget implementation
+	virtual FInputRayHit CanBeginClickDragSequence(const FInputDeviceRay& PressPos) override;
+	virtual void OnClickPress(const FInputDeviceRay& PressPos) override;
+	virtual void OnClickDrag(const FInputDeviceRay& DragPos) override;
+	virtual void OnClickRelease(const FInputDeviceRay& ReleasePos) override;
+	virtual void OnTerminateDragSequence() override;
+	virtual void OnUpdateModifierState(int ModifierID, bool bIsOn) override;
+
+
+	// IHoverBehaviorTarget implementation
+	
+	virtual FInputRayHit BeginHoverSequenceHitTest(const FInputDeviceRay& PressPos) override;
+	virtual void OnBeginHover(const FInputDeviceRay& DevicePos) override {}
+	virtual bool OnUpdateHover(const FInputDeviceRay& DevicePos) override { return true; }
+	virtual void OnEndHover() override {}
+
+
+	// IInteractiveToolCameraFocusAPI implementation
+	virtual bool GetWorldSpaceFocusPoint(const FRay& WorldRay, FVector& PointOut) override;
+
+
+	/** Called by registered InputBehaviors to set the state of the "shift" button (or device equivalent) */
+	virtual void SetShiftToggle(bool bShiftDown);
+
+	/** @return current state of the shift toggle */
+	virtual bool GetShiftToggle() const { return bShiftToggle; }
+
+	/** Called by registered InputBehaviors to set the state of the "shift" button (or device equivalent) */
+	virtual void SetCtrlToggle(bool bCtrlDown);
+
+	/** @return current state of the shift toggle */
+	virtual bool GetCtrlToggle() const { return bCtrlToggle; }
+
+	/** @return current input device pressure in range 0-1 */
+	virtual float GetCurrentDevicePressure() const;
+
+protected:
+	/** Current state of the shift modifier toggle */
+	bool bShiftToggle = false;
+
+	/** Current state of the ctrl modifier toggle */
+	bool bCtrlToggle = false;
+
+	FRay LastWorldRay;
+
+	IToolStylusStateProviderAPI* StylusAPI = nullptr;
+};
+
