InteractiveToolsFramework: Issue complete instead of cancel when switching away from a non accept/cancel tool.

#rnx

diff --git a/Engine/Source/Runtime/InteractiveToolsFramework/Private/InteractiveToolManager.cpp b/Engine/Source/Runtime/InteractiveToolsFramework/Private/InteractiveToolManager.cpp
--- a/Engine/Source/Runtime/InteractiveToolsFramework/Private/InteractiveToolManager.cpp
+++ b/Engine/Source/Runtime/InteractiveToolsFramework/Private/InteractiveToolManager.cpp
@@ -148,7 +148,7 @@ bool UInteractiveToolManager::ActivateTool(EToolSide Side)
 
 	if (ActiveLeftTool != nullptr)
 	{
-		DeactivateTool(EToolSide::Left, ActiveLeftTool->CanAccept() ? EToolShutdownType::Accept : EToolShutdownType::Cancel);
+		DeactivateTool(EToolSide::Left, ActiveLeftTool->CanAccept() ? EToolShutdownType::Accept : EToolShutdownType::Completed);
 	}
 
 	if (ActiveLeftBuilder == nullptr || ActivateToolInternal(Side) == false)
