fix crash in case where transform gizmo somehow doesn't have an associated gizmo actor

#rnx

diff --git a/Engine/Source/Runtime/InteractiveToolsFramework/Public/BaseGizmos/TransformGizmo.h b/Engine/Source/Runtime/InteractiveToolsFramework/Public/BaseGizmos/TransformGizmo.h
--- a/Engine/Source/Runtime/InteractiveToolsFramework/Public/BaseGizmos/TransformGizmo.h
+++ b/Engine/Source/Runtime/InteractiveToolsFramework/Public/BaseGizmos/TransformGizmo.h
@@ -365,7 +365,7 @@ public:
 	 */
 	virtual bool IsVisible()
 	{
-		return !GizmoActor->IsHidden();
+		return GizmoActor && !GizmoActor->IsHidden();
 	}
 
 	/**
