ToolsFramework: add SingleClickOrDragBehavior, switches from click-behavior to drag-behavior if the mouse moves beyond a threshold distance
#rnx

diff --git a/Engine/Source/Runtime/InteractiveToolsFramework/Private/BaseBehaviors/SingleClickOrDragBehavior.cpp b/Engine/Source/Runtime/InteractiveToolsFramework/Private/BaseBehaviors/SingleClickOrDragBehavior.cpp
new file mode 100644
--- /dev/null
+++ b/Engine/Source/Runtime/InteractiveToolsFramework/Private/BaseBehaviors/SingleClickOrDragBehavior.cpp
@@ -0,0 +1,144 @@
+// Copyright Epic Games, Inc. All Rights Reserved.
+
+#include "BaseBehaviors/SingleClickOrDragBehavior.h"
+
+
+
+void USingleClickOrDragInputBehavior::Initialize(IClickBehaviorTarget* ClickTargetIn, IClickDragBehaviorTarget* DragTargetIn)
+{
+	this->ClickTarget = ClickTargetIn;
+	this->DragTarget = DragTargetIn;
+}
+
+
+FInputCaptureRequest USingleClickOrDragInputBehavior::WantsCapture(const FInputDeviceState& Input)
+{
+	bImmediatelyBeginDragInBeginCapture = false;
+	if (IsPressed(Input) && (ModifierCheckFunc == nullptr || ModifierCheckFunc(Input)) )
+	{
+		FInputRayHit ClickHitResult = ClickTarget->IsHitByClick(GetDeviceRay(Input));
+		if (ClickHitResult.bHit)
+		{
+			return FInputCaptureRequest::Begin(this, EInputCaptureSide::Any, ClickHitResult.HitDepth);
+		}
+		else if (bBeginDragIfClickTargetNotHit)
+		{
+			FInputRayHit DragHitResult = DragTarget->CanBeginClickDragSequence(GetDeviceRay(Input));
+			if (DragHitResult.bHit)
+			{
+				bImmediatelyBeginDragInBeginCapture = true;
+				return FInputCaptureRequest::Begin(this, EInputCaptureSide::Any, DragHitResult.HitDepth);
+			}
+		}
+	}
+	return FInputCaptureRequest::Ignore();
+}
+
+
+FInputCaptureUpdate USingleClickOrDragInputBehavior::BeginCapture(const FInputDeviceState& Input, EInputCaptureSide eSide)
+{
+	Modifiers.UpdateModifiers(Input, ClickTarget);
+	Modifiers.UpdateModifiers(Input, DragTarget);
+
+	ensure(Input.IsFromDevice(EInputDevices::Mouse));	// todo: handle other devices
+	MouseDownPosition2D = Input.Mouse.Position2D;	
+	MouseDownRay = Input.Mouse.WorldRay;
+	CaptureSide = eSide;
+	bInDrag = false;
+
+	// If the WantsCapture() hit-test did not hit the target, then start the drag action directly, instead of starting
+	// a click and then switching to a drag (could alternately repeat the hit-test here...)
+	if (bImmediatelyBeginDragInBeginCapture)
+	{
+		bInDrag = true;
+		OnClickDragPressInternal(Input, CaptureSide);
+	}
+
+	return FInputCaptureUpdate::Begin(this, EInputCaptureSide::Any);
+}
+
+
+FInputCaptureUpdate USingleClickOrDragInputBehavior::UpdateCapture(const FInputDeviceState& Input, const FInputCaptureData& Data)
+{
+	Modifiers.UpdateModifiers(Input, ClickTarget);
+	Modifiers.UpdateModifiers(Input, DragTarget);
+
+	// check if mouse has moved far enough that we want to swap to drag behavior
+	if (bInDrag == false)
+	{
+		float MouseMovement = FVector2D::Distance(Input.Mouse.Position2D, MouseDownPosition2D);
+		if (MouseMovement > ClickDistanceThreshold)
+		{
+			FInputDeviceState StartInput = Input;
+			StartInput.Mouse.Position2D = MouseDownPosition2D;		// attempt to reconstruct the input state that would have existed in WantsCapture/BeginCapture
+			StartInput.Mouse.WorldRay = MouseDownRay;
+			FInputRayHit DragHitResult = DragTarget->CanBeginClickDragSequence(GetDeviceRay(StartInput));
+			if (DragHitResult.bHit)
+			{
+				bInDrag = true;
+				OnClickDragPressInternal(StartInput, CaptureSide);
+			}
+		}
+	}
+
+	if (IsReleased(Input))
+	{
+		if (bInDrag)
+		{
+			OnClickDragReleaseInternal(Input, Data);
+			bInDrag = false;
+		}
+		else   // click path
+		{
+			if (ClickTarget->IsHitByClick(GetDeviceRay(Input)).bHit)
+			{
+				OnClickedInternal(Input, Data);
+			}
+		}
+
+		return FInputCaptureUpdate::End();
+	}
+	else
+	{
+		if (bInDrag)
+		{
+			OnClickDragInternal(Input, Data);
+		}
+		return FInputCaptureUpdate::Continue();
+	}
+}
+
+
+void USingleClickOrDragInputBehavior::ForceEndCapture(const FInputCaptureData& data)
+{
+	if (bInDrag)
+	{
+		DragTarget->OnTerminateDragSequence();
+		bInDrag = false;
+	}
+
+	// nothing to do
+}
+
+
+void USingleClickOrDragInputBehavior::OnClickedInternal(const FInputDeviceState& input, const FInputCaptureData& data)
+{
+	ClickTarget->OnClicked(GetDeviceRay(input));
+}
+
+
+
+void USingleClickOrDragInputBehavior::OnClickDragPressInternal(const FInputDeviceState& Input, EInputCaptureSide Side)
+{
+	DragTarget->OnClickPress(GetDeviceRay(Input));
+}
+
+void USingleClickOrDragInputBehavior::OnClickDragInternal(const FInputDeviceState& Input, const FInputCaptureData& Data)
+{
+	DragTarget->OnClickDrag(GetDeviceRay(Input));
+}
+
+void USingleClickOrDragInputBehavior::OnClickDragReleaseInternal(const FInputDeviceState& Input, const FInputCaptureData& Data)
+{
+	DragTarget->OnClickRelease(GetDeviceRay(Input));
+}
diff --git a/Engine/Source/Runtime/InteractiveToolsFramework/Public/BaseBehaviors/SingleClickOrDragBehavior.h b/Engine/Source/Runtime/InteractiveToolsFramework/Public/BaseBehaviors/SingleClickOrDragBehavior.h
new file mode 100644
--- /dev/null
+++ b/Engine/Source/Runtime/InteractiveToolsFramework/Public/BaseBehaviors/SingleClickOrDragBehavior.h
@@ -0,0 +1,117 @@
+// Copyright Epic Games, Inc. All Rights Reserved.
+
+#pragma once
+
+#include "CoreMinimal.h"
+#include "BaseBehaviors/AnyButtonInputBehavior.h"
+#include "BehaviorTargetInterfaces.h"
+#include "InputBehaviorModifierStates.h"
+#include "SingleClickOrDragBehavior.generated.h"
+
+
+/**
+ * USingleClickOrDragInputBehavior is a combination of a USingleClickBehavior and UClickDragBehavior,
+ * and allows for the common UI interaction where a click-and-release does one action, but if the mouse 
+ * is moved, then a drag interaction is started. For example click-to-select is often combined with
+ * a drag-marquee-rectangle in this way. This can be directly implemented with a UClickDragBehavior but
+ * requires the client to (eg) detect movement thresholds, etc. This class encapsulates all that state/logic.
+ * 
+ * The .ClickDistanceThreshold parameter determines how far the mouse must move (in whatever device units are in use)
+ * to switch from a click to drag interaction
+ * 
+ * The .bBeginDragIfClickTargetNotHit parameter determines if the drag interaction will be immediately initiated 
+ * if the initial 'click' mouse-down does not hit a valid clickable target. Defaults to true. 
+ *
+ * The hit-test and on-clicked functions are provided by a IClickBehaviorTarget instance, while an
+ * IClickDragBehaviorTarget provides the can-click-drag/begin-drag/update-drag/end-drag functionality.
+ */
+UCLASS()
+class INTERACTIVETOOLSFRAMEWORK_API USingleClickOrDragInputBehavior : public UAnyButtonInputBehavior
+{
+	GENERATED_BODY()
+
+public:
+
+	/**
+	* The modifier set for this behavior
+	*/
+	FInputBehaviorModifierStates Modifiers;
+
+
+	/**
+	 * Set the targets for Click and Drag interactions
+	 */
+	virtual void Initialize(IClickBehaviorTarget* ClickTarget, IClickDragBehaviorTarget* DragTarget);
+
+
+	/**
+	 * WantsCapture() will only return capture request if this function returns true (or is null)
+	 */
+	TFunction<bool(const FInputDeviceState&)> ModifierCheckFunc = nullptr;
+
+
+	// UInputBehavior implementation
+
+	virtual FInputCaptureRequest WantsCapture(const FInputDeviceState& Input) override;
+	virtual FInputCaptureUpdate BeginCapture(const FInputDeviceState& Input, EInputCaptureSide eSide) override;
+	virtual FInputCaptureUpdate UpdateCapture(const FInputDeviceState& Input, const FInputCaptureData& Data) override;
+	virtual void ForceEndCapture(const FInputCaptureData& Data) override;
+
+
+
+	/** If true (default), then if the click-mouse-down does not hit a valid click target (determined by IClickBehaviorTarget::IsHitByClick), then the Drag will be initiated */
+	UPROPERTY()
+	bool bBeginDragIfClickTargetNotHit = true;
+
+
+	/** If the device moves more than this distance in 2D (pixel?) units, the interaction switches from click to drag */
+	UPROPERTY()
+	float ClickDistanceThreshold = 5.0;
+
+
+protected:
+	/** Click Target object */
+	IClickBehaviorTarget* ClickTarget;
+
+	/** Drag Target object */
+	IClickDragBehaviorTarget* DragTarget;
+
+	/** Click-down position, to determine if mouse has moved far enough to swap to drag interaction */
+	FVector2D MouseDownPosition2D;
+	FRay MouseDownRay;
+
+	/** Device capture */
+	EInputCaptureSide CaptureSide;
+
+	/** set to true if we are in an active drag capture, eg after rejecting possible click */
+	bool bInDrag = false;
+
+
+protected:
+
+	// flag used to communicate between WantsCapture and BeginCapture
+	bool bImmediatelyBeginDragInBeginCapture = false;
+
+	/**
+	 * Internal function that forwards click evens to ClickTarget::OnClicked, you can customize behavior here
+	 */
+	virtual void OnClickedInternal(const FInputDeviceState& Input, const FInputCaptureData& Data);
+
+	/**
+	* Internal function that forwards click evens to DragTarget::OnClickPress, you can customize behavior here
+	*/
+	virtual void OnClickDragPressInternal(const FInputDeviceState& Input, EInputCaptureSide Side);
+
+	/**
+	* Internal function that forwards click evens to DragTarget::OnClickDrag, you can customize behavior here
+	*/
+	virtual void OnClickDragInternal(const FInputDeviceState& Input, const FInputCaptureData& Data);
+
+	/**
+	* Internal function that forwards click evens to DragTarget::OnClickRelease, you can customize behavior here
+	*/
+	virtual void OnClickDragReleaseInternal(const FInputDeviceState& Input, const FInputCaptureData& Data);
+
+
+};
+
