ToolsFramework: rename modeling-mode UTransformGizmo to UCombinedTransformGizmo so that Editor can use this class name, update all usage
#rnx

diff --git a/Engine/Plugins/Editor/SequencerAnimTools/Source/SequencerAnimTools/Public/MotionTrailTool.h b/Engine/Plugins/Editor/SequencerAnimTools/Source/SequencerAnimTools/Public/MotionTrailTool.h
--- a/Engine/Plugins/Editor/SequencerAnimTools/Source/SequencerAnimTools/Public/MotionTrailTool.h
+++ b/Engine/Plugins/Editor/SequencerAnimTools/Source/SequencerAnimTools/Public/MotionTrailTool.h
@@ -21,7 +21,7 @@
 
 class USingleClickInputBehavior;
 class UClickDragInputBehavior;
-class UTransformGizmo;
+class UCombinedTransformGizmo;
 class FEditorModeTools;
 class USingleKeyCaptureBehavior;
 
@@ -159,7 +159,7 @@ public:
 	UPROPERTY()
 	TObjectPtr<UTransformProxy> TransformProxy;
 	UPROPERTY()
-	TObjectPtr<UTransformGizmo> TransformGizmo;
+	TObjectPtr<UCombinedTransformGizmo> TransformGizmo;
 
 	FTransform StartDragTransform;
 	bool bGizmoBeingDragged = false;
diff --git a/Engine/Plugins/Editor/SequencerAnimTools/Source/SequencerAnimTools/Public/SequencerAnimEditPivotTool.h b/Engine/Plugins/Editor/SequencerAnimTools/Source/SequencerAnimTools/Public/SequencerAnimEditPivotTool.h
--- a/Engine/Plugins/Editor/SequencerAnimTools/Source/SequencerAnimTools/Public/SequencerAnimEditPivotTool.h
+++ b/Engine/Plugins/Editor/SequencerAnimTools/Source/SequencerAnimTools/Public/SequencerAnimEditPivotTool.h
@@ -17,7 +17,7 @@
 
 class USingleClickInputBehavior;
 class UClickDragInputBehavior;
-class UTransformGizmo;
+class UCombinedTransformGizmo;
 class ULevelSequence;
 class UControlRig;
 struct FRigControlElement;
@@ -150,7 +150,7 @@ protected:
 	TObjectPtr<UTransformProxy> TransformProxy = nullptr;
 
 	UPROPERTY()
-	TObjectPtr<UTransformGizmo> TransformGizmo = nullptr;
+	TObjectPtr<UCombinedTransformGizmo> TransformGizmo = nullptr;
 
 protected:
 	bool bShiftPressedWhenStarted = false;
diff --git a/Engine/Plugins/Experimental/ChaosEditor/Source/FractureEditor/Private/FractureToolCutter.h b/Engine/Plugins/Experimental/ChaosEditor/Source/FractureEditor/Private/FractureToolCutter.h
--- a/Engine/Plugins/Experimental/ChaosEditor/Source/FractureEditor/Private/FractureToolCutter.h
+++ b/Engine/Plugins/Experimental/ChaosEditor/Source/FractureEditor/Private/FractureToolCutter.h
@@ -248,7 +248,7 @@ public:
 	bool bShowUseGizmoOption = true;
 
 	UPROPERTY()
-	TObjectPtr<UTransformGizmo> TransformGizmo = nullptr;
+	TObjectPtr<UCombinedTransformGizmo> TransformGizmo = nullptr;
 
 	UPROPERTY()
 	TObjectPtr<UTransformProxy> TransformProxy = nullptr;
diff --git a/Engine/Plugins/Experimental/GizmoEdMode/Source/GizmoEdMode/Private/DefaultAssetEditorGizmoFactory.cpp b/Engine/Plugins/Experimental/GizmoEdMode/Source/GizmoEdMode/Private/DefaultAssetEditorGizmoFactory.cpp
--- a/Engine/Plugins/Experimental/GizmoEdMode/Source/GizmoEdMode/Private/DefaultAssetEditorGizmoFactory.cpp
+++ b/Engine/Plugins/Experimental/GizmoEdMode/Source/GizmoEdMode/Private/DefaultAssetEditorGizmoFactory.cpp
@@ -36,7 +36,7 @@ TArray<UInteractiveGizmo*> UDefaultAssetEditorGizmoFactory::BuildGizmoForSelecti
 		Elements = ETransformGizmoSubElements::FullTranslateRotateScale;
 		break;
 	}
-	UTransformGizmo* TransformGizmo = GizmoManager->CreateCustomTransformGizmo(Elements);
+	UCombinedTransformGizmo* TransformGizmo = GizmoManager->CreateCustomTransformGizmo(Elements);
 	TransformGizmo->bUseContextCoordinateSystem = bUseContextCoordinateSystem;
 
 	TArray<AActor*> SelectedActors;
@@ -62,7 +62,7 @@ void UDefaultAssetEditorGizmoFactory::ConfigureGridSnapping(bool bGridEnabled, b
 {
 	for (auto& Gizmo : Gizmos)
 	{
-		UTransformGizmo* TransformGizmo = Cast<UTransformGizmo>(Gizmo);
+		UCombinedTransformGizmo* TransformGizmo = Cast<UCombinedTransformGizmo>(Gizmo);
 
 		if (TransformGizmo)
 		{
diff --git a/Engine/Plugins/Experimental/GizmoEdMode/Source/GizmoEdMode/Public/GizmoEdMode.h b/Engine/Plugins/Experimental/GizmoEdMode/Source/GizmoEdMode/Public/GizmoEdMode.h
--- a/Engine/Plugins/Experimental/GizmoEdMode/Source/GizmoEdMode/Public/GizmoEdMode.h
+++ b/Engine/Plugins/Experimental/GizmoEdMode/Source/GizmoEdMode/Public/GizmoEdMode.h
@@ -10,7 +10,7 @@
 
 #include "GizmoEdMode.generated.h"
 
-class UTransformGizmo;
+class UCombinedTransformGizmo;
 class UInteractiveGizmo;
 class UInteractiveGizmoManager;
 
diff --git a/Engine/Plugins/Experimental/GizmoEdMode/Source/LightGizmos/Private/DirectionalLightGizmoFactory.cpp b/Engine/Plugins/Experimental/GizmoEdMode/Source/LightGizmos/Private/DirectionalLightGizmoFactory.cpp
--- a/Engine/Plugins/Experimental/GizmoEdMode/Source/LightGizmos/Private/DirectionalLightGizmoFactory.cpp
+++ b/Engine/Plugins/Experimental/GizmoEdMode/Source/LightGizmos/Private/DirectionalLightGizmoFactory.cpp
@@ -74,7 +74,7 @@ TArray<UInteractiveGizmo*> UDirectionalLightGizmoFactory::BuildGizmoForSelection
 		Elements = ETransformGizmoSubElements::FullTranslateRotateScale;
 		break;
 	}
-	UTransformGizmo* TransformGizmo = GizmoManager->CreateCustomTransformGizmo(Elements);
+	UCombinedTransformGizmo* TransformGizmo = GizmoManager->CreateCustomTransformGizmo(Elements);
 	TransformGizmo->bUseContextCoordinateSystem = bUseContextCoordinateSystem;
 
 	USubTransformProxy* TransformProxy = NewObject<USubTransformProxy>();
@@ -104,7 +104,7 @@ void UDirectionalLightGizmoFactory::ConfigureGridSnapping(bool bGridEnabled, boo
 {
 	for (UInteractiveGizmo* Gizmo : Gizmos)
 	{
-		UTransformGizmo* TransformGizmo = Cast<UTransformGizmo>(Gizmo);
+		UCombinedTransformGizmo* TransformGizmo = Cast<UCombinedTransformGizmo>(Gizmo);
 
 		if (TransformGizmo)
 		{
diff --git a/Engine/Plugins/Experimental/GizmoEdMode/Source/LightGizmos/Private/PointLightGizmoFactory.cpp b/Engine/Plugins/Experimental/GizmoEdMode/Source/LightGizmos/Private/PointLightGizmoFactory.cpp
--- a/Engine/Plugins/Experimental/GizmoEdMode/Source/LightGizmos/Private/PointLightGizmoFactory.cpp
+++ b/Engine/Plugins/Experimental/GizmoEdMode/Source/LightGizmos/Private/PointLightGizmoFactory.cpp
@@ -76,7 +76,7 @@ TArray<UInteractiveGizmo*> UPointLightGizmoFactory::BuildGizmoForSelection(FEdit
 		break;
 	}
 
-	UTransformGizmo* TransformGizmo = GizmoManager->CreateCustomTransformGizmo(Elements);
+	UCombinedTransformGizmo* TransformGizmo = GizmoManager->CreateCustomTransformGizmo(Elements);
 	TransformGizmo->bUseContextCoordinateSystem = bUseContextCoordinateSystem;
 	
 	USubTransformProxy* TransformProxy = NewObject<USubTransformProxy>();
@@ -105,7 +105,7 @@ void UPointLightGizmoFactory::ConfigureGridSnapping(bool bGridEnabled, bool bRot
 {
 	for (UInteractiveGizmo *Gizmo : Gizmos)
 	{
-		UTransformGizmo* TransformGizmo = Cast<UTransformGizmo>(Gizmo);
+		UCombinedTransformGizmo* TransformGizmo = Cast<UCombinedTransformGizmo>(Gizmo);
 
 		if (TransformGizmo)
 		{
diff --git a/Engine/Plugins/Experimental/GizmoEdMode/Source/LightGizmos/Private/SpotLightGizmoFactory.cpp b/Engine/Plugins/Experimental/GizmoEdMode/Source/LightGizmos/Private/SpotLightGizmoFactory.cpp
--- a/Engine/Plugins/Experimental/GizmoEdMode/Source/LightGizmos/Private/SpotLightGizmoFactory.cpp
+++ b/Engine/Plugins/Experimental/GizmoEdMode/Source/LightGizmos/Private/SpotLightGizmoFactory.cpp
@@ -86,7 +86,7 @@ TArray<UInteractiveGizmo*> USpotLightGizmoFactory::BuildGizmoForSelection(FEdito
 		Elements = ETransformGizmoSubElements::FullTranslateRotateScale;
 		break;
 	}
-	UTransformGizmo* TransformGizmo = GizmoManager->CreateCustomTransformGizmo(Elements);
+	UCombinedTransformGizmo* TransformGizmo = GizmoManager->CreateCustomTransformGizmo(Elements);
 	TransformGizmo->bUseContextCoordinateSystem = bUseContextCoordinateSystem;
 
 	USubTransformProxy* TransformProxy = NewObject<USubTransformProxy>();
@@ -115,7 +115,7 @@ void USpotLightGizmoFactory::ConfigureGridSnapping(bool bGridEnabled, bool bRotG
 {
 	for (UInteractiveGizmo* Gizmo : Gizmos)
 	{
-		UTransformGizmo* TransformGizmo = Cast<UTransformGizmo>(Gizmo);
+		UCombinedTransformGizmo* TransformGizmo = Cast<UCombinedTransformGizmo>(Gizmo);
 
 		if (TransformGizmo)
 		{
diff --git a/Engine/Plugins/Experimental/MeshModelingToolsetExp/Source/MeshModelingToolsEditorOnly/Public/AddPivotActorTool.h b/Engine/Plugins/Experimental/MeshModelingToolsetExp/Source/MeshModelingToolsEditorOnly/Public/AddPivotActorTool.h
--- a/Engine/Plugins/Experimental/MeshModelingToolsetExp/Source/MeshModelingToolsEditorOnly/Public/AddPivotActorTool.h
+++ b/Engine/Plugins/Experimental/MeshModelingToolsetExp/Source/MeshModelingToolsEditorOnly/Public/AddPivotActorTool.h
@@ -11,7 +11,7 @@
 #include "AddPivotActorTool.generated.h"
 
 class UDragAlignmentMechanic;
-class UTransformGizmo;
+class UCombinedTransformGizmo;
 class UTransformProxy;
 
 
@@ -56,7 +56,7 @@ protected:
 	TObjectPtr<UDragAlignmentMechanic> DragAlignmentMechanic = nullptr;
 
 	UPROPERTY()
-	TObjectPtr<UTransformGizmo> TransformGizmo = nullptr;
+	TObjectPtr<UCombinedTransformGizmo> TransformGizmo = nullptr;
 
 	UPROPERTY()
 	TObjectPtr<UTransformProxy> TransformProxy = nullptr;
diff --git a/Engine/Plugins/Experimental/MeshModelingToolsetExp/Source/MeshModelingToolsEditorOnly/Public/PolygonOnMeshTool.h b/Engine/Plugins/Experimental/MeshModelingToolsetExp/Source/MeshModelingToolsEditorOnly/Public/PolygonOnMeshTool.h
--- a/Engine/Plugins/Experimental/MeshModelingToolsetExp/Source/MeshModelingToolsEditorOnly/Public/PolygonOnMeshTool.h
+++ b/Engine/Plugins/Experimental/MeshModelingToolsetExp/Source/MeshModelingToolsEditorOnly/Public/PolygonOnMeshTool.h
@@ -21,7 +21,6 @@
 // predeclarations
 struct FMeshDescription;
 class UDynamicMeshComponent;
-class UTransformGizmo;
 class UTransformProxy;
 class ULineSetComponent;
 
diff --git a/Engine/Plugins/Experimental/MeshModelingToolsetExp/Source/MeshModelingToolsExp/Public/CubeGridTool.h b/Engine/Plugins/Experimental/MeshModelingToolsetExp/Source/MeshModelingToolsExp/Public/CubeGridTool.h
--- a/Engine/Plugins/Experimental/MeshModelingToolsetExp/Source/MeshModelingToolsExp/Public/CubeGridTool.h
+++ b/Engine/Plugins/Experimental/MeshModelingToolsetExp/Source/MeshModelingToolsExp/Public/CubeGridTool.h
@@ -36,7 +36,7 @@ class UMeshOpPreviewWithBackgroundCompute;
 class UMouseHoverBehavior;
 class UNewMeshMaterialProperties;
 class UPreviewGeometry;
-class UTransformGizmo;
+class UCombinedTransformGizmo;
 class UTransformProxy;
 class UToolTarget;
 
@@ -352,7 +352,7 @@ public:
 protected:
 
 	UPROPERTY()
-	TObjectPtr<UTransformGizmo> GridGizmo = nullptr;
+	TObjectPtr<UCombinedTransformGizmo> GridGizmo = nullptr;
 
 	UPROPERTY()
 	TObjectPtr<UDragAlignmentMechanic> GridGizmoAlignmentMechanic = nullptr;
diff --git a/Engine/Plugins/Experimental/MeshModelingToolsetExp/Source/MeshModelingToolsExp/Public/DrawPolyPathTool.h b/Engine/Plugins/Experimental/MeshModelingToolsetExp/Source/MeshModelingToolsExp/Public/DrawPolyPathTool.h
--- a/Engine/Plugins/Experimental/MeshModelingToolsetExp/Source/MeshModelingToolsExp/Public/DrawPolyPathTool.h
+++ b/Engine/Plugins/Experimental/MeshModelingToolsetExp/Source/MeshModelingToolsExp/Public/DrawPolyPathTool.h
@@ -15,7 +15,6 @@
 #include "DrawPolyPathTool.generated.h"
 
 class FMeshVertexChangeBuilder;
-class UTransformGizmo;
 class UTransformProxy;
 
 /**
diff --git a/Engine/Plugins/Experimental/MeshModelingToolsetExp/Source/MeshModelingToolsExp/Public/DrawPolygonTool.h b/Engine/Plugins/Experimental/MeshModelingToolsetExp/Source/MeshModelingToolsExp/Public/DrawPolygonTool.h
--- a/Engine/Plugins/Experimental/MeshModelingToolsetExp/Source/MeshModelingToolsExp/Public/DrawPolygonTool.h
+++ b/Engine/Plugins/Experimental/MeshModelingToolsetExp/Source/MeshModelingToolsExp/Public/DrawPolygonTool.h
@@ -17,7 +17,7 @@
 #include "DrawPolygonTool.generated.h"
 
 
-class UTransformGizmo;
+class UCombinedTransformGizmo;
 class UTransformProxy;
 PREDECLARE_USE_GEOMETRY_CLASS(FDynamicMesh3);
 
@@ -254,7 +254,7 @@ protected:
 	// drawing plane gizmo
 
 	UPROPERTY()
-	TObjectPtr<UTransformGizmo> PlaneTransformGizmo;
+	TObjectPtr<UCombinedTransformGizmo> PlaneTransformGizmo;
 
 	UPROPERTY()
 	TObjectPtr<UTransformProxy> PlaneTransformProxy;
diff --git a/Engine/Plugins/Experimental/MeshModelingToolsetExp/Source/MeshModelingToolsExp/Public/DynamicMeshSculptTool.h b/Engine/Plugins/Experimental/MeshModelingToolsetExp/Source/MeshModelingToolsExp/Public/DynamicMeshSculptTool.h
--- a/Engine/Plugins/Experimental/MeshModelingToolsetExp/Source/MeshModelingToolsExp/Public/DynamicMeshSculptTool.h
+++ b/Engine/Plugins/Experimental/MeshModelingToolsetExp/Source/MeshModelingToolsExp/Public/DynamicMeshSculptTool.h
@@ -23,7 +23,7 @@
 #include "Util/UniqueIndexSet.h"
 #include "DynamicMeshSculptTool.generated.h"
 
-class UTransformGizmo;
+class UCombinedTransformGizmo;
 class UTransformProxy;
 class UMaterialInstanceDynamic;
 
@@ -473,7 +473,7 @@ private:
 
 	// plane gizmo
 	UPROPERTY()
-	TObjectPtr<UTransformGizmo> PlaneTransformGizmo;
+	TObjectPtr<UCombinedTransformGizmo> PlaneTransformGizmo;
 
 	UPROPERTY()
 	TObjectPtr<UTransformProxy> PlaneTransformProxy;
diff --git a/Engine/Plugins/Experimental/MeshModelingToolsetExp/Source/MeshModelingToolsExp/Public/EditMeshPolygonsTool.h b/Engine/Plugins/Experimental/MeshModelingToolsetExp/Source/MeshModelingToolsExp/Public/EditMeshPolygonsTool.h
--- a/Engine/Plugins/Experimental/MeshModelingToolsetExp/Source/MeshModelingToolsExp/Public/EditMeshPolygonsTool.h
+++ b/Engine/Plugins/Experimental/MeshModelingToolsetExp/Source/MeshModelingToolsExp/Public/EditMeshPolygonsTool.h
@@ -31,7 +31,7 @@ class UPolyEditInsetOutsetActivity;
 class UPolyEditCutFacesActivity;
 class UPolyEditPlanarProjectionUVActivity;
 class UPolygonSelectionMechanic;
-class UTransformGizmo;	
+class UCombinedTransformGizmo;	
 class UTransformProxy;
 
 
@@ -493,7 +493,7 @@ protected:
 	TObjectPtr<UDragAlignmentMechanic> DragAlignmentMechanic = nullptr;
 
 	UPROPERTY()
-	TObjectPtr<UTransformGizmo> TransformGizmo = nullptr;
+	TObjectPtr<UCombinedTransformGizmo> TransformGizmo = nullptr;
 
 	UPROPERTY()
 	TObjectPtr<UTransformProxy> TransformProxy = nullptr;
diff --git a/Engine/Plugins/Experimental/MeshModelingToolsetExp/Source/MeshModelingToolsExp/Public/EditPivotTool.h b/Engine/Plugins/Experimental/MeshModelingToolsetExp/Source/MeshModelingToolsExp/Public/EditPivotTool.h
--- a/Engine/Plugins/Experimental/MeshModelingToolsetExp/Source/MeshModelingToolsExp/Public/EditPivotTool.h
+++ b/Engine/Plugins/Experimental/MeshModelingToolsetExp/Source/MeshModelingToolsExp/Public/EditPivotTool.h
@@ -14,7 +14,7 @@
 class UDragAlignmentMechanic;
 class UBaseAxisTranslationGizmo;
 class UAxisAngleGizmo;
-class UTransformGizmo;
+class UCombinedTransformGizmo;
 class UTransformProxy;
 
 /**
@@ -82,7 +82,7 @@ struct FEditPivotTarget
 	TObjectPtr<UTransformProxy> TransformProxy = nullptr;
 
 	UPROPERTY()
-	TObjectPtr<UTransformGizmo> TransformGizmo = nullptr;
+	TObjectPtr<UCombinedTransformGizmo> TransformGizmo = nullptr;
 };
 
 
diff --git a/Engine/Plugins/Experimental/MeshModelingToolsetExp/Source/MeshModelingToolsExp/Public/MeshSpaceDeformerTool.h b/Engine/Plugins/Experimental/MeshModelingToolsetExp/Source/MeshModelingToolsExp/Public/MeshSpaceDeformerTool.h
--- a/Engine/Plugins/Experimental/MeshModelingToolsetExp/Source/MeshModelingToolsExp/Public/MeshSpaceDeformerTool.h
+++ b/Engine/Plugins/Experimental/MeshModelingToolsetExp/Source/MeshModelingToolsExp/Public/MeshSpaceDeformerTool.h
@@ -21,7 +21,7 @@ class UMeshSpaceDeformerTool;
 
 class UDragAlignmentMechanic;
 class UPreviewMesh;
-class UTransformGizmo;
+class UCombinedTransformGizmo;
 class UTransformProxy;
 class UIntervalGizmo;
 class UMeshOpPreviewWithBackgroundCompute;
@@ -251,7 +251,7 @@ protected:
 	TObjectPtr<UIntervalGizmo> IntervalGizmo;
 
 	UPROPERTY()
-	TObjectPtr<UTransformGizmo> TransformGizmo;
+	TObjectPtr<UCombinedTransformGizmo> TransformGizmo;
 
 	UPROPERTY()
 	TObjectPtr<UTransformProxy> TransformProxy;
diff --git a/Engine/Plugins/Experimental/MeshModelingToolsetExp/Source/MeshModelingToolsExp/Public/MeshVertexSculptTool.h b/Engine/Plugins/Experimental/MeshModelingToolsetExp/Source/MeshModelingToolsExp/Public/MeshVertexSculptTool.h
--- a/Engine/Plugins/Experimental/MeshModelingToolsetExp/Source/MeshModelingToolsExp/Public/MeshVertexSculptTool.h
+++ b/Engine/Plugins/Experimental/MeshModelingToolsetExp/Source/MeshModelingToolsExp/Public/MeshVertexSculptTool.h
@@ -18,8 +18,6 @@
 #include "Polygroups/PolygroupSet.h"
 #include "MeshVertexSculptTool.generated.h"
 
-class UTransformGizmo;
-class UTransformProxy;
 class UMaterialInstanceDynamic;
 class FMeshVertexChangeBuilder;
 class UPreviewMesh;
diff --git a/Engine/Plugins/Experimental/MeshModelingToolsetExp/Source/MeshModelingToolsExp/Public/MirrorTool.h b/Engine/Plugins/Experimental/MeshModelingToolsetExp/Source/MeshModelingToolsExp/Public/MirrorTool.h
--- a/Engine/Plugins/Experimental/MeshModelingToolsetExp/Source/MeshModelingToolsExp/Public/MirrorTool.h
+++ b/Engine/Plugins/Experimental/MeshModelingToolsetExp/Source/MeshModelingToolsExp/Public/MirrorTool.h
@@ -15,9 +15,6 @@
 
 #include "MirrorTool.generated.h"
 
-class UTransformGizmo;
-class UTransformProxy;
-
 UCLASS()
 class MESHMODELINGTOOLSEXP_API UMirrorToolBuilder : public UInteractiveToolBuilder
 {
diff --git a/Engine/Plugins/Experimental/MeshModelingToolsetExp/Source/MeshModelingToolsExp/Public/PlaneCutTool.h b/Engine/Plugins/Experimental/MeshModelingToolsetExp/Source/MeshModelingToolsExp/Public/PlaneCutTool.h
--- a/Engine/Plugins/Experimental/MeshModelingToolsetExp/Source/MeshModelingToolsExp/Public/PlaneCutTool.h
+++ b/Engine/Plugins/Experimental/MeshModelingToolsetExp/Source/MeshModelingToolsExp/Public/PlaneCutTool.h
@@ -14,11 +14,6 @@
 #include "PlaneCutTool.generated.h"
 
 
-// predeclarations
-class UTransformGizmo;
-class UTransformProxy;
-
-
 /**
  *
  */
diff --git a/Engine/Plugins/Experimental/MeshModelingToolsetExp/Source/MeshModelingToolsExp/Public/Sculpting/MeshSculptToolBase.h b/Engine/Plugins/Experimental/MeshModelingToolsetExp/Source/MeshModelingToolsExp/Public/Sculpting/MeshSculptToolBase.h
--- a/Engine/Plugins/Experimental/MeshModelingToolsetExp/Source/MeshModelingToolsExp/Public/Sculpting/MeshSculptToolBase.h
+++ b/Engine/Plugins/Experimental/MeshModelingToolsetExp/Source/MeshModelingToolsExp/Public/Sculpting/MeshSculptToolBase.h
@@ -14,7 +14,7 @@
 
 
 class UMaterialInstanceDynamic;
-class UTransformGizmo;
+class UCombinedTransformGizmo;
 class UTransformProxy;
 class UPreviewMesh;
 
@@ -471,7 +471,7 @@ protected:
 public:
 	// plane gizmo
 	UPROPERTY()
-	TObjectPtr<UTransformGizmo> PlaneTransformGizmo;
+	TObjectPtr<UCombinedTransformGizmo> PlaneTransformGizmo;
 
 	UPROPERTY()
 	TObjectPtr<UTransformProxy> PlaneTransformProxy;
diff --git a/Engine/Plugins/Experimental/MeshModelingToolsetExp/Source/MeshModelingToolsExp/Public/TransformMeshesTool.h b/Engine/Plugins/Experimental/MeshModelingToolsetExp/Source/MeshModelingToolsExp/Public/TransformMeshesTool.h
--- a/Engine/Plugins/Experimental/MeshModelingToolsetExp/Source/MeshModelingToolsExp/Public/TransformMeshesTool.h
+++ b/Engine/Plugins/Experimental/MeshModelingToolsetExp/Source/MeshModelingToolsExp/Public/TransformMeshesTool.h
@@ -13,7 +13,7 @@
 class UBaseAxisTranslationGizmo;
 class UAxisAngleGizmo;
 class UDragAlignmentMechanic;
-class UTransformGizmo;
+class UCombinedTransformGizmo;
 class UTransformProxy;
 
 
@@ -141,7 +141,7 @@ struct FTransformMeshesTarget
 	TObjectPtr<UTransformProxy> TransformProxy = nullptr;
 
 	UPROPERTY()
-	TObjectPtr<UTransformGizmo> TransformGizmo = nullptr;
+	TObjectPtr<UCombinedTransformGizmo> TransformGizmo = nullptr;
 };
 
 
diff --git a/Engine/Plugins/Experimental/MeshModelingToolsetExp/Source/MeshModelingToolsExp/Public/UVProjectionTool.h b/Engine/Plugins/Experimental/MeshModelingToolsetExp/Source/MeshModelingToolsExp/Public/UVProjectionTool.h
--- a/Engine/Plugins/Experimental/MeshModelingToolsetExp/Source/MeshModelingToolsExp/Public/UVProjectionTool.h
+++ b/Engine/Plugins/Experimental/MeshModelingToolsetExp/Source/MeshModelingToolsExp/Public/UVProjectionTool.h
@@ -19,7 +19,7 @@
 // predeclarations
 struct FMeshDescription;
 class UDynamicMeshComponent;
-class UTransformGizmo;
+class UCombinedTransformGizmo;
 class UTransformProxy;
 class USingleClickInputBehavior;
 
@@ -273,7 +273,7 @@ protected:
 	TObjectPtr<UMaterialInstanceDynamic> CheckerMaterial = nullptr;
 
 	UPROPERTY()
-	TObjectPtr<UTransformGizmo> TransformGizmo = nullptr;
+	TObjectPtr<UCombinedTransformGizmo> TransformGizmo = nullptr;
 	
 	UPROPERTY()
 	TObjectPtr<UTransformProxy> TransformProxy = nullptr;
diff --git a/Engine/Plugins/Experimental/MeshModelingToolsetExp/Source/ModelingComponents/Private/BaseTools/BaseCreateFromSelectedTool.cpp b/Engine/Plugins/Experimental/MeshModelingToolsetExp/Source/ModelingComponents/Private/BaseTools/BaseCreateFromSelectedTool.cpp
--- a/Engine/Plugins/Experimental/MeshModelingToolsetExp/Source/ModelingComponents/Private/BaseTools/BaseCreateFromSelectedTool.cpp
+++ b/Engine/Plugins/Experimental/MeshModelingToolsetExp/Source/ModelingComponents/Private/BaseTools/BaseCreateFromSelectedTool.cpp
@@ -167,7 +167,7 @@ void UBaseCreateFromSelectedTool::UpdateGizmoVisibility()
 {
 	for (int32 GizmoIndex = 0; GizmoIndex < TransformGizmos.Num(); GizmoIndex++)
 	{
-		UTransformGizmo* Gizmo = TransformGizmos[GizmoIndex];
+		UCombinedTransformGizmo* Gizmo = TransformGizmos[GizmoIndex];
 		Gizmo->SetVisibility(TransformProperties->bShowTransformUI && GizmoIndex != GetHiddenGizmoIndex());
 	}
 }
@@ -181,7 +181,7 @@ void UBaseCreateFromSelectedTool::SetTransformGizmos()
 	for (int ComponentIdx = 0; ComponentIdx < Targets.Num(); ComponentIdx++)
 	{
 		UTransformProxy* Proxy = TransformProxies.Add_GetRef(NewObject<UTransformProxy>(this));
-		UTransformGizmo* Gizmo = TransformGizmos.Add_GetRef(UE::TransformGizmoUtil::Create3AxisTransformGizmo(GizmoManager, this));
+		UCombinedTransformGizmo* Gizmo = TransformGizmos.Add_GetRef(UE::TransformGizmoUtil::Create3AxisTransformGizmo(GizmoManager, this));
 		Proxy->SetTransform(TargetComponentInterface(ComponentIdx)->GetWorldTransform());
 		Gizmo->SetActiveTarget(Proxy, GetToolManager());
 		Proxy->OnTransformChanged.AddUObject(this, &UBaseCreateFromSelectedTool::TransformChanged);
diff --git a/Engine/Plugins/Experimental/MeshModelingToolsetExp/Source/ModelingComponents/Private/Mechanics/DragAlignmentMechanic.cpp b/Engine/Plugins/Experimental/MeshModelingToolsetExp/Source/ModelingComponents/Private/Mechanics/DragAlignmentMechanic.cpp
--- a/Engine/Plugins/Experimental/MeshModelingToolsetExp/Source/ModelingComponents/Private/Mechanics/DragAlignmentMechanic.cpp
+++ b/Engine/Plugins/Experimental/MeshModelingToolsetExp/Source/ModelingComponents/Private/Mechanics/DragAlignmentMechanic.cpp
@@ -70,7 +70,7 @@ void UDragAlignmentMechanic::Shutdown()
 {
 }
 
-void UDragAlignmentMechanic::AddToGizmo(UTransformGizmo* TransformGizmo, const TArray<const UPrimitiveComponent*>* ComponentsToIgnoreInAlignment, 
+void UDragAlignmentMechanic::AddToGizmo(UCombinedTransformGizmo* TransformGizmo, const TArray<const UPrimitiveComponent*>* ComponentsToIgnoreInAlignment, 
 	const TArray<const UPrimitiveComponent*>* InvisibleComponentsToIncludeInAlignment)
 {
 	// If we have components to ignore/include, we need a copy of the array so that the alignment
diff --git a/Engine/Plugins/Experimental/MeshModelingToolsetExp/Source/ModelingComponents/Public/BaseTools/BaseCreateFromSelectedTool.h b/Engine/Plugins/Experimental/MeshModelingToolsetExp/Source/ModelingComponents/Public/BaseTools/BaseCreateFromSelectedTool.h
--- a/Engine/Plugins/Experimental/MeshModelingToolsetExp/Source/ModelingComponents/Public/BaseTools/BaseCreateFromSelectedTool.h
+++ b/Engine/Plugins/Experimental/MeshModelingToolsetExp/Source/ModelingComponents/Public/BaseTools/BaseCreateFromSelectedTool.h
@@ -10,7 +10,7 @@
 #include "PropertySets/CreateMeshObjectTypeProperties.h"
 #include "BaseCreateFromSelectedTool.generated.h"
 
-class UTransformGizmo;
+class UCombinedTransformGizmo;
 class UTransformProxy;
 
 /**
@@ -201,7 +201,7 @@ protected:
 	TArray<TObjectPtr<UTransformProxy>> TransformProxies;
 
 	UPROPERTY()
-	TArray<TObjectPtr<UTransformGizmo>> TransformGizmos;
+	TArray<TObjectPtr<UCombinedTransformGizmo>> TransformGizmos;
 
 	UWorld* TargetWorld = nullptr;
 };
diff --git a/Engine/Plugins/Experimental/MeshModelingToolsetExp/Source/ModelingComponents/Public/Mechanics/ConstructionPlaneMechanic.h b/Engine/Plugins/Experimental/MeshModelingToolsetExp/Source/ModelingComponents/Public/Mechanics/ConstructionPlaneMechanic.h
--- a/Engine/Plugins/Experimental/MeshModelingToolsetExp/Source/ModelingComponents/Public/Mechanics/ConstructionPlaneMechanic.h
+++ b/Engine/Plugins/Experimental/MeshModelingToolsetExp/Source/ModelingComponents/Public/Mechanics/ConstructionPlaneMechanic.h
@@ -11,7 +11,7 @@
 
 #include "ConstructionPlaneMechanic.generated.h"
 
-class UTransformGizmo;
+class UCombinedTransformGizmo;
 class UTransformProxy;
 class IClickBehaviorTarget;
 class USingleClickInputBehavior;
@@ -60,7 +60,7 @@ public:
 
 public:
 	UPROPERTY()
-	TObjectPtr<UTransformGizmo> PlaneTransformGizmo;
+	TObjectPtr<UCombinedTransformGizmo> PlaneTransformGizmo;
 	
 	UPROPERTY()
 	TObjectPtr<UTransformProxy> PlaneTransformProxy;
diff --git a/Engine/Plugins/Experimental/MeshModelingToolsetExp/Source/ModelingComponents/Public/Mechanics/CurveControlPointsMechanic.h b/Engine/Plugins/Experimental/MeshModelingToolsetExp/Source/ModelingComponents/Public/Mechanics/CurveControlPointsMechanic.h
--- a/Engine/Plugins/Experimental/MeshModelingToolsetExp/Source/ModelingComponents/Public/Mechanics/CurveControlPointsMechanic.h
+++ b/Engine/Plugins/Experimental/MeshModelingToolsetExp/Source/ModelingComponents/Public/Mechanics/CurveControlPointsMechanic.h
@@ -21,7 +21,7 @@ class ULineSetComponent;
 class UMouseHoverBehavior;
 class UPointSetComponent;
 class USingleClickInputBehavior;
-class UTransformGizmo;
+class UCombinedTransformGizmo;
 class UTransformProxy;
 
 /**
@@ -379,7 +379,7 @@ protected:
 	UPROPERTY()
 	TObjectPtr<UTransformProxy> PointTransformProxy;
 	UPROPERTY()
-	TObjectPtr<UTransformGizmo> PointTransformGizmo;
+	TObjectPtr<UCombinedTransformGizmo> PointTransformGizmo;
 
 	// Used to make it easy to tell whether the gizmo was moved by the user or by undo/redo or
 	// some other change that we shoulnd't respond to. Basing our movement undo/redo on the
diff --git a/Engine/Plugins/Experimental/MeshModelingToolsetExp/Source/ModelingComponents/Public/Mechanics/DragAlignmentMechanic.h b/Engine/Plugins/Experimental/MeshModelingToolsetExp/Source/ModelingComponents/Public/Mechanics/DragAlignmentMechanic.h
--- a/Engine/Plugins/Experimental/MeshModelingToolsetExp/Source/ModelingComponents/Public/Mechanics/DragAlignmentMechanic.h
+++ b/Engine/Plugins/Experimental/MeshModelingToolsetExp/Source/ModelingComponents/Public/Mechanics/DragAlignmentMechanic.h
@@ -17,10 +17,10 @@
 PREDECLARE_USE_GEOMETRY_CLASS(FGroupTopologyDeformer);
 class UPrimitiveComponent;
 class UIntervalGizmo;
-class UTransformGizmo;
+class UCombinedTransformGizmo;
 
 /**
- * Mechanic that can be added to (potentially multiple) UTransformGizmo object to allow them to snap to
+ * Mechanic that can be added to (potentially multiple) UCombinedTransformGizmo object to allow them to snap to
  * scene geometry in rotation and translation when the Ctrl key is pressed.
  */
 UCLASS()
@@ -31,13 +31,13 @@ class MODELINGCOMPONENTS_API UDragAlignmentMechanic : public UInteractionMechani
 public:
 
 	/**
-	 * Adds this mechanic to the given gizmo. Can be called on multiple gizmos, and works both for UTransformGizmo
+	 * Adds this mechanic to the given gizmo. Can be called on multiple gizmos, and works both for UCombinedTransformGizmo
 	 * and URepositionableTransformGizmo.
 	 * For repositioning, the gizmo will not ignore anything in its alignment (so that the gizmo can be aligned to the
 	 * target object(s)). For movement, the gizmo will ignore ComponentsToIgnore, as well as modified triangles if used
 	 * for deformation.
 	 */
-	void AddToGizmo(UTransformGizmo* TransformGizmo, const TArray<const UPrimitiveComponent*>* ComponentsToIgnore = nullptr,
+	void AddToGizmo(UCombinedTransformGizmo* TransformGizmo, const TArray<const UPrimitiveComponent*>* ComponentsToIgnore = nullptr,
 		const TArray<const UPrimitiveComponent*>* InvisibleComponentsToInclude = nullptr);
 
 	/**
diff --git a/Engine/Plugins/Experimental/MeshModelingToolsetExp/Source/ModelingComponents/Public/Mechanics/LatticeControlPointsMechanic.h b/Engine/Plugins/Experimental/MeshModelingToolsetExp/Source/ModelingComponents/Public/Mechanics/LatticeControlPointsMechanic.h
--- a/Engine/Plugins/Experimental/MeshModelingToolsetExp/Source/ModelingComponents/Public/Mechanics/LatticeControlPointsMechanic.h
+++ b/Engine/Plugins/Experimental/MeshModelingToolsetExp/Source/ModelingComponents/Public/Mechanics/LatticeControlPointsMechanic.h
@@ -24,7 +24,7 @@ class ULineSetComponent;
 class UMouseHoverBehavior;
 class UPointSetComponent;
 class USingleClickInputBehavior;
-class UTransformGizmo;
+class UCombinedTransformGizmo;
 class UTransformProxy;
 
 UCLASS()
@@ -151,7 +151,7 @@ protected:
 	UPROPERTY()
 	TObjectPtr<UTransformProxy> PointTransformProxy;
 	UPROPERTY()
-	TObjectPtr<UTransformGizmo> PointTransformGizmo;
+	TObjectPtr<UCombinedTransformGizmo> PointTransformGizmo;
 
 	// Used to make it easy to tell whether the gizmo was moved by the user or by undo/redo or
 	// some other change that we shouldn't respond to. Basing our movement undo/redo on the
diff --git a/Engine/Plugins/Experimental/MeshModelingToolsetExp/Source/ModelingComponents/Public/Mechanics/SpaceCurveDeformationMechanic.h b/Engine/Plugins/Experimental/MeshModelingToolsetExp/Source/ModelingComponents/Public/Mechanics/SpaceCurveDeformationMechanic.h
--- a/Engine/Plugins/Experimental/MeshModelingToolsetExp/Source/ModelingComponents/Public/Mechanics/SpaceCurveDeformationMechanic.h
+++ b/Engine/Plugins/Experimental/MeshModelingToolsetExp/Source/ModelingComponents/Public/Mechanics/SpaceCurveDeformationMechanic.h
@@ -21,7 +21,7 @@ class ULineSetComponent;
 class UMouseHoverBehavior;
 class UPointSetComponent;
 class USingleClickInputBehavior;
-class UTransformGizmo;
+class UCombinedTransformGizmo;
 class UTransformProxy;
 
 
@@ -191,7 +191,7 @@ protected:
 	UPROPERTY()
 	TObjectPtr<UTransformProxy> PointTransformProxy;
 	UPROPERTY()
-	TObjectPtr<UTransformGizmo> PointTransformGizmo;
+	TObjectPtr<UCombinedTransformGizmo> PointTransformGizmo;
 
 	// Used to make it easy to tell whether the gizmo was moved by the user or by undo/redo or
 	// some other change that we shoulnd't respond to. Basing our movement undo/redo on the
diff --git a/Engine/Plugins/Experimental/MeshModelingToolsetExp/Source/ModelingComponents/Public/Transforms/MultiTransformer.h b/Engine/Plugins/Experimental/MeshModelingToolsetExp/Source/ModelingComponents/Public/Transforms/MultiTransformer.h
--- a/Engine/Plugins/Experimental/MeshModelingToolsetExp/Source/ModelingComponents/Public/Transforms/MultiTransformer.h
+++ b/Engine/Plugins/Experimental/MeshModelingToolsetExp/Source/ModelingComponents/Public/Transforms/MultiTransformer.h
@@ -14,7 +14,7 @@
 #include "MultiTransformer.generated.h"
 
 
-class UTransformGizmo;
+class UCombinedTransformGizmo;
 class UTransformProxy;
 
 UENUM()
@@ -112,7 +112,7 @@ protected:
 	bool bDisallowNegativeScaling = false;
 
 	UPROPERTY()
-	TObjectPtr<UTransformGizmo> TransformGizmo;
+	TObjectPtr<UCombinedTransformGizmo> TransformGizmo;
 
 	UPROPERTY()
 	TObjectPtr<UTransformProxy> TransformProxy;
diff --git a/Engine/Plugins/Experimental/MotionTrailEditorMode/Source/MotionTrailEditorMode/Private/Sequencer/MovieSceneTransformTrailTool.cpp b/Engine/Plugins/Experimental/MotionTrailEditorMode/Source/MotionTrailEditorMode/Private/Sequencer/MovieSceneTransformTrailTool.cpp
--- a/Engine/Plugins/Experimental/MotionTrailEditorMode/Source/MotionTrailEditorMode/Private/Sequencer/MovieSceneTransformTrailTool.cpp
+++ b/Engine/Plugins/Experimental/MotionTrailEditorMode/Source/MotionTrailEditorMode/Private/Sequencer/MovieSceneTransformTrailTool.cpp
@@ -124,10 +124,10 @@ void FDefaultMovieSceneTransformTrailTool::OnClicked(const FInputDeviceRay& Clic
 
 	if (CachedSelected)
 	{
-		ActiveTransformGizmo = Cast<UTransformGizmo>(TrailToolManager->GetGizmoManager()->FindGizmoByInstanceIdentifier(UTrailToolManager::TrailKeyTransformGizmoInstanceIdentifier));
+		ActiveTransformGizmo = Cast<UCombinedTransformGizmo>(TrailToolManager->GetGizmoManager()->FindGizmoByInstanceIdentifier(UTrailToolManager::TrailKeyTransformGizmoInstanceIdentifier));
 
 		UMSTrailTransformProxy* MSTrailTransformProxy;
-		if (ActiveTransformGizmo.IsValid() && Cast<UMSTrailTransformProxy>(Cast<UTransformGizmo>(ActiveTransformGizmo)->ActiveTarget))
+		if (ActiveTransformGizmo.IsValid() && Cast<UMSTrailTransformProxy>(Cast<UCombinedTransformGizmo>(ActiveTransformGizmo)->ActiveTarget))
 		{
 			MSTrailTransformProxy = Cast<UMSTrailTransformProxy>(ActiveTransformGizmo->ActiveTarget);
 
@@ -431,7 +431,7 @@ void FDefaultMovieSceneTransformTrailTool::FKeyInfo::UpdateKeyTransform(EKeyUpda
 	}
 }
 
-void FDefaultMovieSceneTransformTrailTool::UpdateGizmoActorComponents(FKeyInfo* KeyInfo, UTransformGizmo* TransformGizmo)
+void FDefaultMovieSceneTransformTrailTool::UpdateGizmoActorComponents(FKeyInfo* KeyInfo, UCombinedTransformGizmo* TransformGizmo)
 {
 	if (!KeyInfo->IdxMap.Contains(EMSTrailTransformChannel::TranslateX))
 	{
diff --git a/Engine/Plugins/Experimental/MotionTrailEditorMode/Source/MotionTrailEditorMode/Private/Sequencer/MovieSceneTransformTrailTool.h b/Engine/Plugins/Experimental/MotionTrailEditorMode/Source/MotionTrailEditorMode/Private/Sequencer/MovieSceneTransformTrailTool.h
--- a/Engine/Plugins/Experimental/MotionTrailEditorMode/Source/MotionTrailEditorMode/Private/Sequencer/MovieSceneTransformTrailTool.h
+++ b/Engine/Plugins/Experimental/MotionTrailEditorMode/Source/MotionTrailEditorMode/Private/Sequencer/MovieSceneTransformTrailTool.h
@@ -101,13 +101,13 @@ public:
 
 private:
 
-	void UpdateGizmoActorComponents(FKeyInfo* KeyInfo, class UTransformGizmo* TransformGizmo);
+	void UpdateGizmoActorComponents(FKeyInfo* KeyInfo, class UCombinedTransformGizmo* TransformGizmo);
 
 	static UMSTrailKeyProperties* KeyProps;
 
 	TMap<FFrameNumber, TUniquePtr<FKeyInfo>> Keys;
 	FKeyInfo* CachedSelected;
-	TWeakObjectPtr<class UTransformGizmo> ActiveTransformGizmo;
+	TWeakObjectPtr<class UCombinedTransformGizmo> ActiveTransformGizmo;
 	FMovieSceneTransformTrail* OwningTrail;
 
 	friend class ::UMSTrailTransformProxy;
diff --git a/Engine/Plugins/Experimental/UVEditor/Source/UVEditorTools/Public/UVSelectTool.h b/Engine/Plugins/Experimental/UVEditor/Source/UVEditorTools/Public/UVSelectTool.h
--- a/Engine/Plugins/Experimental/UVEditor/Source/UVEditorTools/Public/UVSelectTool.h
+++ b/Engine/Plugins/Experimental/UVEditor/Source/UVEditorTools/Public/UVSelectTool.h
@@ -21,7 +21,7 @@ class FToolCommandChange;
 class ULineSetComponent;
 class UMeshSelectionMechanic;
 class UToolTargetManager;
-class UTransformGizmo;
+class UCombinedTransformGizmo;
 class UTransformProxy;
 class UUVEditorMode;
 class UUVSelectToolChangeRouter;
@@ -177,7 +177,7 @@ protected:
 	TObjectPtr<UMeshSelectionMechanic> SelectionMechanic = nullptr;
 
 	UPROPERTY()
-	TObjectPtr<UTransformGizmo> TransformGizmo = nullptr;
+	TObjectPtr<UCombinedTransformGizmo> TransformGizmo = nullptr;
 
 	TArray<TSharedPtr<FDynamicMeshAABBTree3>> AABBTrees;
 
diff --git a/Engine/Source/Runtime/InteractiveToolsFramework/Private/BaseGizmos/GizmoActor.cpp b/Engine/Source/Runtime/InteractiveToolsFramework/Private/BaseGizmos/GizmoActor.cpp
--- a/Engine/Source/Runtime/InteractiveToolsFramework/Private/BaseGizmos/GizmoActor.cpp
+++ b/Engine/Source/Runtime/InteractiveToolsFramework/Private/BaseGizmos/GizmoActor.cpp
@@ -12,7 +12,7 @@
 #include "Engine/World.h"
 
 
-#define LOCTEXT_NAMESPACE "UTransformGizmo"
+#define LOCTEXT_NAMESPACE "AGizmoActor"
 
 
 AGizmoActor::AGizmoActor()
diff --git a/Engine/Source/Runtime/InteractiveToolsFramework/Private/BaseGizmos/RepositionableTransformGizmo.cpp b/Engine/Source/Runtime/InteractiveToolsFramework/Private/BaseGizmos/RepositionableTransformGizmo.cpp
--- a/Engine/Source/Runtime/InteractiveToolsFramework/Private/BaseGizmos/RepositionableTransformGizmo.cpp
+++ b/Engine/Source/Runtime/InteractiveToolsFramework/Private/BaseGizmos/RepositionableTransformGizmo.cpp
@@ -19,14 +19,14 @@
 
 UInteractiveGizmo* URepositionableTransformGizmoBuilder::BuildGizmo(const FToolBuilderState& SceneState) const
 {
-	UTransformGizmo* NewGizmo = NewObject<URepositionableTransformGizmo>(SceneState.GizmoManager);
+	UCombinedTransformGizmo* NewGizmo = NewObject<URepositionableTransformGizmo>(SceneState.GizmoManager);
 	NewGizmo->SetWorld(SceneState.World);
 
 	UGizmoViewContext* GizmoViewContext = SceneState.ToolManager->GetContextObjectStore()->FindContext<UGizmoViewContext>();
 	check(GizmoViewContext && GizmoViewContext->IsValidLowLevel());
 
 	// use default gizmo actor if client has not given us a new builder
-	NewGizmo->SetGizmoActorBuilder(GizmoActorBuilder ? GizmoActorBuilder : MakeShared<FTransformGizmoActorFactory>(GizmoViewContext));
+	NewGizmo->SetGizmoActorBuilder(GizmoActorBuilder ? GizmoActorBuilder : MakeShared<FCombinedTransformGizmoActorFactory>(GizmoViewContext));
 
 	NewGizmo->SetSubGizmoBuilderIdentifiers(AxisPositionBuilderIdentifier, PlanePositionBuilderIdentifier, AxisAngleBuilderIdentifier);
 
@@ -51,7 +51,7 @@ void URepositionableTransformGizmo::SetActiveTarget(UTransformProxy* Target, ITo
 	// We're going to add on to the original transform gizmo, adding extra gizmos centered around
 	// existing gizmo components. These will respond to middle clicks and take "set pivot" code
 	// paths in their transform source and state target.
-	UTransformGizmo::SetActiveTarget(Target, TransactionProvider);
+	UCombinedTransformGizmo::SetActiveTarget(Target, TransactionProvider);
 
 	// Create the alternate transform source.
 	UGizmoTransformProxyTransformSource* ProxyTransformSource = UGizmoTransformProxyTransformSource::Construct(ActiveTarget, this);
@@ -206,7 +206,7 @@ void URepositionableTransformGizmo::SetWorldAlignmentFunctions(
 	TUniqueFunction<bool(const FRay&, FVector&)>&& TranslationAlignmentRayCasterIn)
 {
 	// The base class function will do most of what we want.
-	UTransformGizmo::SetWorldAlignmentFunctions(MoveTemp(ShouldAlignTranslationIn), MoveTemp(TranslationAlignmentRayCasterIn));
+	UCombinedTransformGizmo::SetWorldAlignmentFunctions(MoveTemp(ShouldAlignTranslationIn), MoveTemp(TranslationAlignmentRayCasterIn));
 
 	// However, we need to adjust our saved pivot alignment functions so changes of the gizmo target keep the settings.
 	// Since the functions are TUniqueFunction, we can't actually copy them here, we need to create new functors that
diff --git a/Engine/Source/Runtime/InteractiveToolsFramework/Private/BaseGizmos/TransformGizmo.cpp b/Engine/Source/Runtime/InteractiveToolsFramework/Private/BaseGizmos/TransformGizmo.cpp
--- a/Engine/Source/Runtime/InteractiveToolsFramework/Private/BaseGizmos/TransformGizmo.cpp
+++ b/Engine/Source/Runtime/InteractiveToolsFramework/Private/BaseGizmos/TransformGizmo.cpp
@@ -24,10 +24,10 @@
 #include "Engine/CollisionProfile.h"
 
 
-#define LOCTEXT_NAMESPACE "UTransformGizmo"
+#define LOCTEXT_NAMESPACE "UCombinedTransformGizmo"
 
 
-ATransformGizmoActor::ATransformGizmoActor()
+ACombinedTransformGizmoActor::ACombinedTransformGizmoActor()
 {
 	// root component is a hidden sphere
 	USphereComponent* SphereComponent = CreateDefaultSubobject<USphereComponent>(TEXT("GizmoCenter"));
@@ -39,7 +39,7 @@ ATransformGizmoActor::ATransformGizmoActor()
 
 
 
-ATransformGizmoActor* ATransformGizmoActor::ConstructDefault3AxisGizmo(UWorld* World, UGizmoViewContext* GizmoViewContext)
+ACombinedTransformGizmoActor* ACombinedTransformGizmoActor::ConstructDefault3AxisGizmo(UWorld* World, UGizmoViewContext* GizmoViewContext)
 {
 	return ConstructCustom3AxisGizmo(World, GizmoViewContext,
 		ETransformGizmoSubElements::TranslateAllAxes |
@@ -52,12 +52,12 @@ ATransformGizmoActor* ATransformGizmoActor::ConstructDefault3AxisGizmo(UWorld* W
 }
 
 
-ATransformGizmoActor* ATransformGizmoActor::ConstructCustom3AxisGizmo(
+ACombinedTransformGizmoActor* ACombinedTransformGizmoActor::ConstructCustom3AxisGizmo(
 	UWorld* World, UGizmoViewContext* GizmoViewContext,
 	ETransformGizmoSubElements Elements)
 {
 	FActorSpawnParameters SpawnInfo;
-	ATransformGizmoActor* NewActor = World->SpawnActor<ATransformGizmoActor>(FVector::ZeroVector, FRotator::ZeroRotator, SpawnInfo);
+	ACombinedTransformGizmoActor* NewActor = World->SpawnActor<ACombinedTransformGizmoActor>(FVector::ZeroVector, FRotator::ZeroRotator, SpawnInfo);
 
 	float GizmoLineThickness = 3.0f;
 
@@ -213,23 +213,23 @@ ATransformGizmoActor* ATransformGizmoActor::ConstructCustom3AxisGizmo(
 
 
 
-ATransformGizmoActor* FTransformGizmoActorFactory::CreateNewGizmoActor(UWorld* World) const
+ACombinedTransformGizmoActor* FCombinedTransformGizmoActorFactory::CreateNewGizmoActor(UWorld* World) const
 {
-	return ATransformGizmoActor::ConstructCustom3AxisGizmo(World, GizmoViewContext, EnableElements);
+	return ACombinedTransformGizmoActor::ConstructCustom3AxisGizmo(World, GizmoViewContext, EnableElements);
 }
 
 
 
-UInteractiveGizmo* UTransformGizmoBuilder::BuildGizmo(const FToolBuilderState& SceneState) const
+UInteractiveGizmo* UCombinedTransformGizmoBuilder::BuildGizmo(const FToolBuilderState& SceneState) const
 {
-	UTransformGizmo* NewGizmo = NewObject<UTransformGizmo>(SceneState.GizmoManager);
+	UCombinedTransformGizmo* NewGizmo = NewObject<UCombinedTransformGizmo>(SceneState.GizmoManager);
 	NewGizmo->SetWorld(SceneState.World);
 
 	UGizmoViewContext* GizmoViewContext = SceneState.ToolManager->GetContextObjectStore()->FindContext<UGizmoViewContext>();
 	check(GizmoViewContext && GizmoViewContext->IsValidLowLevel());
 
 	// use default gizmo actor if client has not given us a new builder
-	NewGizmo->SetGizmoActorBuilder(GizmoActorBuilder ? GizmoActorBuilder : MakeShared<FTransformGizmoActorFactory>(GizmoViewContext));
+	NewGizmo->SetGizmoActorBuilder(GizmoActorBuilder ? GizmoActorBuilder : MakeShared<FCombinedTransformGizmoActorFactory>(GizmoViewContext));
 
 	NewGizmo->SetSubGizmoBuilderIdentifiers(AxisPositionBuilderIdentifier, PlanePositionBuilderIdentifier, AxisAngleBuilderIdentifier);
 
@@ -249,34 +249,34 @@ UInteractiveGizmo* UTransformGizmoBuilder::BuildGizmo(const FToolBuilderState& S
 
 
 
-void UTransformGizmo::SetWorld(UWorld* WorldIn)
+void UCombinedTransformGizmo::SetWorld(UWorld* WorldIn)
 {
 	this->World = WorldIn;
 }
 
-void UTransformGizmo::SetGizmoActorBuilder(TSharedPtr<FTransformGizmoActorFactory> Builder)
+void UCombinedTransformGizmo::SetGizmoActorBuilder(TSharedPtr<FCombinedTransformGizmoActorFactory> Builder)
 {
 	GizmoActorBuilder = Builder;
 }
 
-void UTransformGizmo::SetSubGizmoBuilderIdentifiers(FString AxisPositionBuilderIdentifierIn, FString PlanePositionBuilderIdentifierIn, FString AxisAngleBuilderIdentifierIn)
+void UCombinedTransformGizmo::SetSubGizmoBuilderIdentifiers(FString AxisPositionBuilderIdentifierIn, FString PlanePositionBuilderIdentifierIn, FString AxisAngleBuilderIdentifierIn)
 {
 	AxisPositionBuilderIdentifier = AxisPositionBuilderIdentifierIn;
 	PlanePositionBuilderIdentifier = PlanePositionBuilderIdentifierIn;
 	AxisAngleBuilderIdentifier = AxisAngleBuilderIdentifierIn;
 }
 
-void UTransformGizmo::SetUpdateHoverFunction(TFunction<void(UPrimitiveComponent*, bool)> HoverFunction)
+void UCombinedTransformGizmo::SetUpdateHoverFunction(TFunction<void(UPrimitiveComponent*, bool)> HoverFunction)
 {
 	UpdateHoverFunction = HoverFunction;
 }
 
-void UTransformGizmo::SetUpdateCoordSystemFunction(TFunction<void(UPrimitiveComponent*, EToolContextCoordinateSystem)> CoordSysFunction)
+void UCombinedTransformGizmo::SetUpdateCoordSystemFunction(TFunction<void(UPrimitiveComponent*, EToolContextCoordinateSystem)> CoordSysFunction)
 {
 	UpdateCoordSystemFunction = CoordSysFunction;
 }
 
-void UTransformGizmo::SetWorldAlignmentFunctions(
+void UCombinedTransformGizmo::SetWorldAlignmentFunctions(
 	TUniqueFunction<bool()>&& ShouldAlignTranslationIn,
 	TUniqueFunction<bool(const FRay&, FVector&)>&& TranslationAlignmentRayCasterIn)
 {
@@ -333,7 +333,7 @@ void UTransformGizmo::SetWorldAlignmentFunctions(
 	}
 }
 
-void UTransformGizmo::SetDisallowNegativeScaling(bool bDisallow)
+void UCombinedTransformGizmo::SetDisallowNegativeScaling(bool bDisallow)
 {
 	if (bDisallowNegativeScaling != bDisallow)
 	{
@@ -358,13 +358,13 @@ void UTransformGizmo::SetDisallowNegativeScaling(bool bDisallow)
 	}
 }
 
-void UTransformGizmo::SetIsNonUniformScaleAllowedFunction(TUniqueFunction<bool()>&& IsNonUniformScaleAllowedIn)
+void UCombinedTransformGizmo::SetIsNonUniformScaleAllowedFunction(TUniqueFunction<bool()>&& IsNonUniformScaleAllowedIn)
 {
 	IsNonUniformScaleAllowed = MoveTemp(IsNonUniformScaleAllowedIn);
 }
 
 
-void UTransformGizmo::Setup()
+void UCombinedTransformGizmo::Setup()
 {
 	UInteractiveGizmo::Setup();
 
@@ -389,7 +389,7 @@ void UTransformGizmo::Setup()
 
 
 
-void UTransformGizmo::Shutdown()
+void UCombinedTransformGizmo::Shutdown()
 {
 	ClearActiveTarget();
 
@@ -402,7 +402,7 @@ void UTransformGizmo::Shutdown()
 
 
 
-void UTransformGizmo::UpdateCameraAxisSource()
+void UCombinedTransformGizmo::UpdateCameraAxisSource()
 {
 	FViewCameraState CameraState;
 	GetGizmoManager()->GetContextQueriesAPI()->GetCurrentViewState(CameraState);
@@ -416,7 +416,7 @@ void UTransformGizmo::UpdateCameraAxisSource()
 }
 
 
-void UTransformGizmo::Tick(float DeltaTime)
+void UCombinedTransformGizmo::Tick(float DeltaTime)
 {	
 	if (bUseContextCoordinateSystem)
 	{
@@ -451,7 +451,7 @@ void UTransformGizmo::Tick(float DeltaTime)
 
 
 
-void UTransformGizmo::SetActiveTarget(UTransformProxy* Target, IToolContextTransactionProvider* TransactionProvider)
+void UCombinedTransformGizmo::SetActiveTarget(UTransformProxy* Target, IToolContextTransactionProvider* TransactionProvider)
 {
 	if (ActiveTarget != nullptr)
 	{
@@ -479,7 +479,7 @@ void UTransformGizmo::SetActiveTarget(UTransformProxy* Target, IToolContextTrans
 		TransactionProvider = GetGizmoManager();
 	}
 	StateTarget = UGizmoTransformChangeStateTarget::Construct(GizmoComponent,
-		LOCTEXT("UTransformGizmoTransaction", "Transform"), TransactionProvider, this);
+		LOCTEXT("UCombinedTransformGizmoTransaction", "Transform"), TransactionProvider, this);
 	StateTarget->DependentChangeSources.Add(MakeUnique<FTransformProxyChangeSource>(Target));
 
 	CameraAxisSource = NewObject<UGizmoConstantFrameAxisSource>(this);
@@ -592,13 +592,13 @@ void UTransformGizmo::SetActiveTarget(UTransformProxy* Target, IToolContextTrans
 	}
 }
 
-FTransform UTransformGizmo::GetGizmoTransform() const
+FTransform UCombinedTransformGizmo::GetGizmoTransform() const
 {
 	USceneComponent* GizmoComponent = GizmoActor->GetRootComponent();
 	return GizmoComponent->GetComponentTransform();
 }
 
-void UTransformGizmo::ReinitializeGizmoTransform(const FTransform& NewTransform, bool bKeepGizmoUnscaled)
+void UCombinedTransformGizmo::ReinitializeGizmoTransform(const FTransform& NewTransform, bool bKeepGizmoUnscaled)
 {
 	// To update the gizmo location without triggering any callbacks, we temporarily
 	// store a copy of the callback list, detach them, reposition, and then reattach
@@ -621,7 +621,7 @@ void UTransformGizmo::ReinitializeGizmoTransform(const FTransform& NewTransform,
 	ActiveTarget->bSetPivotMode = bSavedSetPivotMode;
 }
 
-void UTransformGizmo::SetNewGizmoTransform(const FTransform& NewTransform, bool bKeepGizmoUnscaled)
+void UCombinedTransformGizmo::SetNewGizmoTransform(const FTransform& NewTransform, bool bKeepGizmoUnscaled)
 {
 	check(ActiveTarget != nullptr);
 
@@ -640,7 +640,7 @@ void UTransformGizmo::SetNewGizmoTransform(const FTransform& NewTransform, bool
 }
 
 
-void UTransformGizmo::SetNewChildScale(const FVector& NewChildScale)
+void UCombinedTransformGizmo::SetNewChildScale(const FVector& NewChildScale)
 {
 	FTransform NewTransform = ActiveTarget->GetTransform();
 	NewTransform.SetScale3D(NewChildScale);
@@ -652,7 +652,7 @@ void UTransformGizmo::SetNewChildScale(const FVector& NewChildScale)
 }
 
 
-void UTransformGizmo::SetVisibility(bool bVisible)
+void UCombinedTransformGizmo::SetVisibility(bool bVisible)
 {
 	GizmoActor->SetActorHiddenInGame(bVisible == false);
 #if WITH_EDITOR
@@ -661,7 +661,7 @@ void UTransformGizmo::SetVisibility(bool bVisible)
 }
 
 
-UInteractiveGizmo* UTransformGizmo::AddAxisTranslationGizmo(
+UInteractiveGizmo* UCombinedTransformGizmo::AddAxisTranslationGizmo(
 	UPrimitiveComponent* AxisComponent, USceneComponent* RootComponent,
 	IGizmoAxisSource* AxisSource,
 	IGizmoTransformSource* TransformSource,
@@ -701,7 +701,7 @@ UInteractiveGizmo* UTransformGizmo::AddAxisTranslationGizmo(
 
 
 
-UInteractiveGizmo* UTransformGizmo::AddPlaneTranslationGizmo(
+UInteractiveGizmo* UCombinedTransformGizmo::AddPlaneTranslationGizmo(
 	UPrimitiveComponent* AxisComponent, USceneComponent* RootComponent,
 	IGizmoAxisSource* AxisSource,
 	IGizmoTransformSource* TransformSource,
@@ -743,7 +743,7 @@ UInteractiveGizmo* UTransformGizmo::AddPlaneTranslationGizmo(
 
 
 
-UInteractiveGizmo* UTransformGizmo::AddAxisRotationGizmo(
+UInteractiveGizmo* UCombinedTransformGizmo::AddAxisRotationGizmo(
 	UPrimitiveComponent* AxisComponent, USceneComponent* RootComponent,
 	IGizmoAxisSource* AxisSource,
 	IGizmoTransformSource* TransformSource,
@@ -784,7 +784,7 @@ UInteractiveGizmo* UTransformGizmo::AddAxisRotationGizmo(
 
 
 
-UInteractiveGizmo* UTransformGizmo::AddAxisScaleGizmo(
+UInteractiveGizmo* UCombinedTransformGizmo::AddAxisScaleGizmo(
 	UPrimitiveComponent* AxisComponent, USceneComponent* RootComponent,
 	IGizmoAxisSource* GizmoAxisSource, IGizmoAxisSource* ParameterAxisSource,
 	IGizmoTransformSource* TransformSource,
@@ -819,7 +819,7 @@ UInteractiveGizmo* UTransformGizmo::AddAxisScaleGizmo(
 
 
 
-UInteractiveGizmo* UTransformGizmo::AddPlaneScaleGizmo(
+UInteractiveGizmo* UCombinedTransformGizmo::AddPlaneScaleGizmo(
 	UPrimitiveComponent* AxisComponent, USceneComponent* RootComponent,
 	IGizmoAxisSource* GizmoAxisSource, IGizmoAxisSource* ParameterAxisSource,
 	IGizmoTransformSource* TransformSource,
@@ -857,7 +857,7 @@ UInteractiveGizmo* UTransformGizmo::AddPlaneScaleGizmo(
 
 
 
-UInteractiveGizmo* UTransformGizmo::AddUniformScaleGizmo(
+UInteractiveGizmo* UCombinedTransformGizmo::AddUniformScaleGizmo(
 	UPrimitiveComponent* ScaleComponent, USceneComponent* RootComponent,
 	IGizmoAxisSource* GizmoAxisSource, IGizmoAxisSource* ParameterAxisSource,
 	IGizmoTransformSource* TransformSource,
@@ -891,7 +891,7 @@ UInteractiveGizmo* UTransformGizmo::AddUniformScaleGizmo(
 
 
 
-void UTransformGizmo::ClearActiveTarget()
+void UCombinedTransformGizmo::ClearActiveTarget()
 {
 	for (UInteractiveGizmo* Gizmo : ActiveGizmos)
 	{
@@ -916,7 +916,7 @@ void UTransformGizmo::ClearActiveTarget()
 
 
 
-bool UTransformGizmo::PositionSnapFunction(const FVector& WorldPosition, FVector& SnappedPositionOut) const
+bool UCombinedTransformGizmo::PositionSnapFunction(const FVector& WorldPosition, FVector& SnappedPositionOut) const
 {
 	SnappedPositionOut = WorldPosition;
 
@@ -951,7 +951,7 @@ bool UTransformGizmo::PositionSnapFunction(const FVector& WorldPosition, FVector
 	return false;
 }
 
-FQuat UTransformGizmo::RotationSnapFunction(const FQuat& DeltaRotation) const
+FQuat UCombinedTransformGizmo::RotationSnapFunction(const FQuat& DeltaRotation) const
 {
 	FQuat SnappedDeltaRotation = DeltaRotation;
 
diff --git a/Engine/Source/Runtime/InteractiveToolsFramework/Private/BaseGizmos/TransformGizmoUtil.cpp b/Engine/Source/Runtime/InteractiveToolsFramework/Private/BaseGizmos/TransformGizmoUtil.cpp
--- a/Engine/Source/Runtime/InteractiveToolsFramework/Private/BaseGizmos/TransformGizmoUtil.cpp
+++ b/Engine/Source/Runtime/InteractiveToolsFramework/Private/BaseGizmos/TransformGizmoUtil.cpp
@@ -14,15 +14,15 @@
 #include "BaseGizmos/TransformGizmo.h"
 #include "BaseGizmos/RepositionableTransformGizmo.h"
 
-const FString UTransformGizmoContextObject::DefaultAxisPositionBuilderIdentifier = TEXT("Util_StandardXFormAxisTranslationGizmo");
-const FString UTransformGizmoContextObject::DefaultPlanePositionBuilderIdentifier = TEXT("Util_StandardXFormPlaneTranslationGizmo");
-const FString UTransformGizmoContextObject::DefaultAxisAngleBuilderIdentifier = TEXT("Util_StandardXFormAxisRotationGizmo");
-const FString UTransformGizmoContextObject::DefaultThreeAxisTransformBuilderIdentifier = TEXT("Util_DefaultThreeAxisTransformBuilderIdentifier");
-const FString UTransformGizmoContextObject::CustomThreeAxisTransformBuilderIdentifier = TEXT("Util_CustomThreeAxisTransformBuilderIdentifier");
-const FString UTransformGizmoContextObject::CustomRepositionableThreeAxisTransformBuilderIdentifier = TEXT("Util_CustomRepositionableThreeAxisTransformBuilderIdentifier");
+const FString UCombinedTransformGizmoContextObject::DefaultAxisPositionBuilderIdentifier = TEXT("Util_StandardXFormAxisTranslationGizmo");
+const FString UCombinedTransformGizmoContextObject::DefaultPlanePositionBuilderIdentifier = TEXT("Util_StandardXFormPlaneTranslationGizmo");
+const FString UCombinedTransformGizmoContextObject::DefaultAxisAngleBuilderIdentifier = TEXT("Util_StandardXFormAxisRotationGizmo");
+const FString UCombinedTransformGizmoContextObject::DefaultThreeAxisTransformBuilderIdentifier = TEXT("Util_DefaultThreeAxisTransformBuilderIdentifier");
+const FString UCombinedTransformGizmoContextObject::CustomThreeAxisTransformBuilderIdentifier = TEXT("Util_CustomThreeAxisTransformBuilderIdentifier");
+const FString UCombinedTransformGizmoContextObject::CustomRepositionableThreeAxisTransformBuilderIdentifier = TEXT("Util_CustomRepositionableThreeAxisTransformBuilderIdentifier");
 
 
-void UTransformGizmoContextObject::RegisterGizmosWithManager(UInteractiveToolManager* ToolManager)
+void UCombinedTransformGizmoContextObject::RegisterGizmosWithManager(UInteractiveToolManager* ToolManager)
 {
 	if (ensure(!bDefaultGizmosRegistered) == false)
 	{
@@ -41,7 +41,7 @@ void UTransformGizmoContextObject::RegisterGizmosWithManager(UInteractiveToolMan
 	UAxisAngleGizmoBuilder* AxisRotationBuilder = NewObject<UAxisAngleGizmoBuilder>();
 	GizmoManager->RegisterGizmoType(DefaultAxisAngleBuilderIdentifier, AxisRotationBuilder);
 
-	UTransformGizmoBuilder* TransformBuilder = NewObject<UTransformGizmoBuilder>();
+	UCombinedTransformGizmoBuilder* TransformBuilder = NewObject<UCombinedTransformGizmoBuilder>();
 	TransformBuilder->AxisPositionBuilderIdentifier = DefaultAxisPositionBuilderIdentifier;
 	TransformBuilder->PlanePositionBuilderIdentifier = DefaultPlanePositionBuilderIdentifier;
 	TransformBuilder->AxisAngleBuilderIdentifier = DefaultAxisAngleBuilderIdentifier;
@@ -54,9 +54,9 @@ void UTransformGizmoContextObject::RegisterGizmosWithManager(UInteractiveToolMan
 		ToolManager->GetContextObjectStore()->AddContextObject(GizmoViewContext);
 	}
 
-	GizmoActorBuilder = MakeShared<FTransformGizmoActorFactory>(GizmoViewContext);
+	GizmoActorBuilder = MakeShared<FCombinedTransformGizmoActorFactory>(GizmoViewContext);
 
-	UTransformGizmoBuilder* CustomThreeAxisBuilder = NewObject<UTransformGizmoBuilder>();
+	UCombinedTransformGizmoBuilder* CustomThreeAxisBuilder = NewObject<UCombinedTransformGizmoBuilder>();
 	CustomThreeAxisBuilder->AxisPositionBuilderIdentifier = DefaultAxisPositionBuilderIdentifier;
 	CustomThreeAxisBuilder->PlanePositionBuilderIdentifier = DefaultPlanePositionBuilderIdentifier;
 	CustomThreeAxisBuilder->AxisAngleBuilderIdentifier = DefaultAxisAngleBuilderIdentifier;
@@ -75,7 +75,7 @@ void UTransformGizmoContextObject::RegisterGizmosWithManager(UInteractiveToolMan
 
 }
 
-void UTransformGizmoContextObject::DeregisterGizmosWithManager(UInteractiveToolManager* ToolManager)
+void UCombinedTransformGizmoContextObject::DeregisterGizmosWithManager(UInteractiveToolManager* ToolManager)
 {
 	UInteractiveGizmoManager* GizmoManager = ToolManager->GetPairedGizmoManager();
 	ToolManager->GetContextObjectStore()->RemoveContextObject(this);
@@ -97,10 +97,10 @@ bool UE::TransformGizmoUtil::RegisterTransformGizmoContextObject(UInteractiveToo
 {
 	if (ensure(ToolsContext))
 	{
-		UTransformGizmoContextObject* Found = ToolsContext->ContextObjectStore->FindContext<UTransformGizmoContextObject>();
+		UCombinedTransformGizmoContextObject* Found = ToolsContext->ContextObjectStore->FindContext<UCombinedTransformGizmoContextObject>();
 		if (Found == nullptr)
 		{
-			UTransformGizmoContextObject* GizmoHelper = NewObject<UTransformGizmoContextObject>(ToolsContext->ToolManager);
+			UCombinedTransformGizmoContextObject* GizmoHelper = NewObject<UCombinedTransformGizmoContextObject>(ToolsContext->ToolManager);
 			if (ensure(GizmoHelper))
 			{
 				GizmoHelper->RegisterGizmosWithManager(ToolsContext->ToolManager);
@@ -121,7 +121,7 @@ bool UE::TransformGizmoUtil::DeregisterTransformGizmoContextObject(UInteractiveT
 {
 	if (ensure(ToolsContext))
 	{
-		UTransformGizmoContextObject* Found = ToolsContext->ContextObjectStore->FindContext<UTransformGizmoContextObject>();
+		UCombinedTransformGizmoContextObject* Found = ToolsContext->ContextObjectStore->FindContext<UCombinedTransformGizmoContextObject>();
 		if (Found != nullptr)
 		{
 			Found->DeregisterGizmosWithManager(ToolsContext->ToolManager);
@@ -134,21 +134,21 @@ bool UE::TransformGizmoUtil::DeregisterTransformGizmoContextObject(UInteractiveT
 
 
 
-UTransformGizmo* UTransformGizmoContextObject::Create3AxisTransformGizmo(UInteractiveGizmoManager* GizmoManager, void* Owner, const FString& InstanceIdentifier)
+UCombinedTransformGizmo* UCombinedTransformGizmoContextObject::Create3AxisTransformGizmo(UInteractiveGizmoManager* GizmoManager, void* Owner, const FString& InstanceIdentifier)
 {
 	if (ensure(bDefaultGizmosRegistered))
 	{
 		UInteractiveGizmo* NewGizmo = GizmoManager->CreateGizmo(DefaultThreeAxisTransformBuilderIdentifier, InstanceIdentifier, Owner);
 		ensure(NewGizmo);
-		return Cast<UTransformGizmo>(NewGizmo);
+		return Cast<UCombinedTransformGizmo>(NewGizmo);
 	}
 	return nullptr;
 }
-UTransformGizmo* UE::TransformGizmoUtil::Create3AxisTransformGizmo(UInteractiveGizmoManager* GizmoManager, void* Owner, const FString& InstanceIdentifier)
+UCombinedTransformGizmo* UE::TransformGizmoUtil::Create3AxisTransformGizmo(UInteractiveGizmoManager* GizmoManager, void* Owner, const FString& InstanceIdentifier)
 {
 	if (ensure(GizmoManager))
 	{
-		UTransformGizmoContextObject* UseThis = GizmoManager->GetContextObjectStore()->FindContext<UTransformGizmoContextObject>();
+		UCombinedTransformGizmoContextObject* UseThis = GizmoManager->GetContextObjectStore()->FindContext<UCombinedTransformGizmoContextObject>();
 		if (ensure(UseThis))
 		{
 			return UseThis->Create3AxisTransformGizmo(GizmoManager, Owner, InstanceIdentifier);
@@ -156,7 +156,7 @@ UTransformGizmo* UE::TransformGizmoUtil::Create3AxisTransformGizmo(UInteractiveG
 	}
 	return nullptr;
 }
-UTransformGizmo* UE::TransformGizmoUtil::Create3AxisTransformGizmo(UInteractiveToolManager* ToolManager, void* Owner, const FString& InstanceIdentifier)
+UCombinedTransformGizmo* UE::TransformGizmoUtil::Create3AxisTransformGizmo(UInteractiveToolManager* ToolManager, void* Owner, const FString& InstanceIdentifier)
 {
 	if (ensure(ToolManager))
 	{
@@ -165,22 +165,22 @@ UTransformGizmo* UE::TransformGizmoUtil::Create3AxisTransformGizmo(UInteractiveT
 	return nullptr;
 }
 
-UTransformGizmo* UTransformGizmoContextObject::CreateCustomTransformGizmo(UInteractiveGizmoManager* GizmoManager, ETransformGizmoSubElements Elements, void* Owner, const FString& InstanceIdentifier)
+UCombinedTransformGizmo* UCombinedTransformGizmoContextObject::CreateCustomTransformGizmo(UInteractiveGizmoManager* GizmoManager, ETransformGizmoSubElements Elements, void* Owner, const FString& InstanceIdentifier)
 {
 	if (ensure(bDefaultGizmosRegistered))
 	{
 		GizmoActorBuilder->EnableElements = Elements;
 		UInteractiveGizmo* NewGizmo = GizmoManager->CreateGizmo(CustomThreeAxisTransformBuilderIdentifier, InstanceIdentifier, Owner);
 		ensure(NewGizmo);
-		return Cast<UTransformGizmo>(NewGizmo);
+		return Cast<UCombinedTransformGizmo>(NewGizmo);
 	}
 	return nullptr;
 }
-UTransformGizmo* UE::TransformGizmoUtil::CreateCustomTransformGizmo(UInteractiveGizmoManager* GizmoManager, ETransformGizmoSubElements Elements, void* Owner, const FString& InstanceIdentifier)
+UCombinedTransformGizmo* UE::TransformGizmoUtil::CreateCustomTransformGizmo(UInteractiveGizmoManager* GizmoManager, ETransformGizmoSubElements Elements, void* Owner, const FString& InstanceIdentifier)
 {
 	if (ensure(GizmoManager))
 	{
-		UTransformGizmoContextObject* UseThis = GizmoManager->GetContextObjectStore()->FindContext<UTransformGizmoContextObject>();
+		UCombinedTransformGizmoContextObject* UseThis = GizmoManager->GetContextObjectStore()->FindContext<UCombinedTransformGizmoContextObject>();
 		if (ensure(UseThis))
 		{
 			return UseThis->CreateCustomTransformGizmo(GizmoManager, Elements, Owner, InstanceIdentifier);
@@ -188,7 +188,7 @@ UTransformGizmo* UE::TransformGizmoUtil::CreateCustomTransformGizmo(UInteractive
 	}
 	return nullptr;
 }
-UTransformGizmo* UE::TransformGizmoUtil::CreateCustomTransformGizmo(UInteractiveToolManager* ToolManager, ETransformGizmoSubElements Elements, void* Owner, const FString& InstanceIdentifier)
+UCombinedTransformGizmo* UE::TransformGizmoUtil::CreateCustomTransformGizmo(UInteractiveToolManager* ToolManager, ETransformGizmoSubElements Elements, void* Owner, const FString& InstanceIdentifier)
 {
 	if (ensure(ToolManager))
 	{
@@ -198,22 +198,22 @@ UTransformGizmo* UE::TransformGizmoUtil::CreateCustomTransformGizmo(UInteractive
 }
 
 
-UTransformGizmo* UTransformGizmoContextObject::CreateCustomRepositionableTransformGizmo(UInteractiveGizmoManager* GizmoManager, ETransformGizmoSubElements Elements, void* Owner, const FString& InstanceIdentifier)
+UCombinedTransformGizmo* UCombinedTransformGizmoContextObject::CreateCustomRepositionableTransformGizmo(UInteractiveGizmoManager* GizmoManager, ETransformGizmoSubElements Elements, void* Owner, const FString& InstanceIdentifier)
 {
 	if (ensure(bDefaultGizmosRegistered))
 	{
 		GizmoActorBuilder->EnableElements = Elements;
 		UInteractiveGizmo* NewGizmo = GizmoManager->CreateGizmo(CustomRepositionableThreeAxisTransformBuilderIdentifier, InstanceIdentifier, Owner);
 		ensure(NewGizmo);
-		return Cast<UTransformGizmo>(NewGizmo);
+		return Cast<UCombinedTransformGizmo>(NewGizmo);
 	}
 	return nullptr;
 }
-UTransformGizmo* UE::TransformGizmoUtil::CreateCustomRepositionableTransformGizmo(UInteractiveGizmoManager* GizmoManager, ETransformGizmoSubElements Elements, void* Owner, const FString& InstanceIdentifier)
+UCombinedTransformGizmo* UE::TransformGizmoUtil::CreateCustomRepositionableTransformGizmo(UInteractiveGizmoManager* GizmoManager, ETransformGizmoSubElements Elements, void* Owner, const FString& InstanceIdentifier)
 {
 	if (ensure(GizmoManager))
 	{
-		UTransformGizmoContextObject* UseThis = GizmoManager->GetContextObjectStore()->FindContext<UTransformGizmoContextObject>();
+		UCombinedTransformGizmoContextObject* UseThis = GizmoManager->GetContextObjectStore()->FindContext<UCombinedTransformGizmoContextObject>();
 		if (ensure(UseThis))
 		{
 			return UseThis->CreateCustomRepositionableTransformGizmo(GizmoManager, Elements, Owner, InstanceIdentifier);
@@ -221,7 +221,7 @@ UTransformGizmo* UE::TransformGizmoUtil::CreateCustomRepositionableTransformGizm
 	}
 	return nullptr;
 }
-UTransformGizmo* UE::TransformGizmoUtil::CreateCustomRepositionableTransformGizmo(UInteractiveToolManager* ToolManager, ETransformGizmoSubElements Elements, void* Owner, const FString& InstanceIdentifier)
+UCombinedTransformGizmo* UE::TransformGizmoUtil::CreateCustomRepositionableTransformGizmo(UInteractiveToolManager* ToolManager, ETransformGizmoSubElements Elements, void* Owner, const FString& InstanceIdentifier)
 {
 	if (ensure(ToolManager))
 	{
diff --git a/Engine/Source/Runtime/InteractiveToolsFramework/Private/InteractiveGizmoManager.cpp b/Engine/Source/Runtime/InteractiveToolsFramework/Private/InteractiveGizmoManager.cpp
--- a/Engine/Source/Runtime/InteractiveToolsFramework/Private/InteractiveGizmoManager.cpp
+++ b/Engine/Source/Runtime/InteractiveToolsFramework/Private/InteractiveGizmoManager.cpp
@@ -296,7 +296,7 @@ void UInteractiveGizmoManager::RegisterDefaultGizmos()
 	UAxisAngleGizmoBuilder* AxisRotationBuilder = NewObject<UAxisAngleGizmoBuilder>();
 	RegisterGizmoType(DefaultAxisAngleBuilderIdentifier, AxisRotationBuilder);
 
-	UTransformGizmoBuilder* TransformBuilder = NewObject<UTransformGizmoBuilder>();
+	UCombinedTransformGizmoBuilder* TransformBuilder = NewObject<UCombinedTransformGizmoBuilder>();
 	TransformBuilder->AxisPositionBuilderIdentifier = DefaultAxisPositionBuilderIdentifier;
 	TransformBuilder->PlanePositionBuilderIdentifier = DefaultPlanePositionBuilderIdentifier;
 	TransformBuilder->AxisAngleBuilderIdentifier = DefaultAxisAngleBuilderIdentifier;
@@ -308,9 +308,9 @@ void UInteractiveGizmoManager::RegisterDefaultGizmos()
 		GizmoViewContext = NewObject<UGizmoViewContext>();
 		GetContextObjectStore()->AddContextObject(GizmoViewContext);
 	}
-	GizmoActorBuilder = MakeShared<FTransformGizmoActorFactory>(GizmoViewContext);
+	GizmoActorBuilder = MakeShared<FCombinedTransformGizmoActorFactory>(GizmoViewContext);
 
-	UTransformGizmoBuilder* CustomThreeAxisBuilder = NewObject<UTransformGizmoBuilder>();
+	UCombinedTransformGizmoBuilder* CustomThreeAxisBuilder = NewObject<UCombinedTransformGizmoBuilder>();
 	CustomThreeAxisBuilder->AxisPositionBuilderIdentifier = DefaultAxisPositionBuilderIdentifier;
 	CustomThreeAxisBuilder->PlanePositionBuilderIdentifier = DefaultPlanePositionBuilderIdentifier;
 	CustomThreeAxisBuilder->AxisAngleBuilderIdentifier = DefaultAxisAngleBuilderIdentifier;
@@ -330,30 +330,30 @@ void UInteractiveGizmoManager::RegisterDefaultGizmos()
 	bDefaultGizmosRegistered = true;
 }
 
-UTransformGizmo* UInteractiveGizmoManager::Create3AxisTransformGizmo(void* Owner, const FString& InstanceIdentifier)
+UCombinedTransformGizmo* UInteractiveGizmoManager::Create3AxisTransformGizmo(void* Owner, const FString& InstanceIdentifier)
 {
 	check(bDefaultGizmosRegistered);
 	UInteractiveGizmo* NewGizmo = CreateGizmo(DefaultThreeAxisTransformBuilderIdentifier, InstanceIdentifier, Owner);
 	check(NewGizmo);
-	return Cast<UTransformGizmo>(NewGizmo);
+	return Cast<UCombinedTransformGizmo>(NewGizmo);
 }
 
-UTransformGizmo* UInteractiveGizmoManager::CreateCustomTransformGizmo(ETransformGizmoSubElements Elements, void* Owner, const FString& InstanceIdentifier)
+UCombinedTransformGizmo* UInteractiveGizmoManager::CreateCustomTransformGizmo(ETransformGizmoSubElements Elements, void* Owner, const FString& InstanceIdentifier)
 {
 	check(bDefaultGizmosRegistered);
 	GizmoActorBuilder->EnableElements = Elements;
 	UInteractiveGizmo* NewGizmo = CreateGizmo(CustomThreeAxisTransformBuilderIdentifier, InstanceIdentifier, Owner);
 	check(NewGizmo);
-	return Cast<UTransformGizmo>(NewGizmo);
+	return Cast<UCombinedTransformGizmo>(NewGizmo);
 }
 
-UTransformGizmo* UInteractiveGizmoManager::CreateCustomRepositionableTransformGizmo(ETransformGizmoSubElements Elements, void* Owner, const FString& InstanceIdentifier)
+UCombinedTransformGizmo* UInteractiveGizmoManager::CreateCustomRepositionableTransformGizmo(ETransformGizmoSubElements Elements, void* Owner, const FString& InstanceIdentifier)
 {
 	check(bDefaultGizmosRegistered);
 	GizmoActorBuilder->EnableElements = Elements;
 	UInteractiveGizmo* NewGizmo = CreateGizmo(CustomRepositionableThreeAxisTransformBuilderIdentifier, InstanceIdentifier, Owner);
 	check(NewGizmo);
-	return Cast<UTransformGizmo>(NewGizmo);
+	return Cast<UCombinedTransformGizmo>(NewGizmo);
 }
 
 
diff --git a/Engine/Source/Runtime/InteractiveToolsFramework/Public/BaseGizmos/IntervalGizmo.h b/Engine/Source/Runtime/InteractiveToolsFramework/Public/BaseGizmos/IntervalGizmo.h
--- a/Engine/Source/Runtime/InteractiveToolsFramework/Public/BaseGizmos/IntervalGizmo.h
+++ b/Engine/Source/Runtime/InteractiveToolsFramework/Public/BaseGizmos/IntervalGizmo.h
@@ -83,7 +83,7 @@ public:
 
 	/**
 	 * @param World the UWorld to create the new Actor in
-	 * @return new ATransformGizmoActor instance with members initialized with Components suitable for a transformation Gizmo
+	 * @return new AIntervalGizmoActor instance with members initialized with Components suitable for a transformation Gizmo
 	 */
 	virtual AIntervalGizmoActor* CreateNewGizmoActor(UWorld* World) const
 	{
diff --git a/Engine/Source/Runtime/InteractiveToolsFramework/Public/BaseGizmos/RepositionableTransformGizmo.h b/Engine/Source/Runtime/InteractiveToolsFramework/Public/BaseGizmos/RepositionableTransformGizmo.h
--- a/Engine/Source/Runtime/InteractiveToolsFramework/Public/BaseGizmos/RepositionableTransformGizmo.h
+++ b/Engine/Source/Runtime/InteractiveToolsFramework/Public/BaseGizmos/RepositionableTransformGizmo.h
@@ -14,7 +14,7 @@ class IToolContextTransactionProvider;
 class UTransformProxy;
 
 UCLASS()
-class INTERACTIVETOOLSFRAMEWORK_API URepositionableTransformGizmoBuilder : public UTransformGizmoBuilder
+class INTERACTIVETOOLSFRAMEWORK_API URepositionableTransformGizmoBuilder : public UCombinedTransformGizmoBuilder
 {
 	GENERATED_BODY()
 
@@ -28,7 +28,7 @@ public:
  * A transform gizmo that also allows the user to reposition it by middle-clicking rotation/translation components.
  */
 UCLASS()
-class INTERACTIVETOOLSFRAMEWORK_API URepositionableTransformGizmo : public UTransformGizmo
+class INTERACTIVETOOLSFRAMEWORK_API URepositionableTransformGizmo : public UCombinedTransformGizmo
 {
 	GENERATED_BODY()
 
diff --git a/Engine/Source/Runtime/InteractiveToolsFramework/Public/BaseGizmos/TransformGizmo.h b/Engine/Source/Runtime/InteractiveToolsFramework/Public/BaseGizmos/TransformGizmo.h
--- a/Engine/Source/Runtime/InteractiveToolsFramework/Public/BaseGizmos/TransformGizmo.h
+++ b/Engine/Source/Runtime/InteractiveToolsFramework/Public/BaseGizmos/TransformGizmo.h
@@ -23,10 +23,10 @@ class UGizmoViewContext;
 class FTransformGizmoTransformChange;
 
 /**
- * ATransformGizmoActor is an Actor type intended to be used with UTransformGizmo,
+ * ACombinedTransformGizmoActor is an Actor type intended to be used with UCombinedTransformGizmo,
  * as the in-scene visual representation of the Gizmo.
  * 
- * FTransformGizmoActorFactory returns an instance of this Actor type (or a subclass), and based on
+ * FCombinedTransformGizmoActorFactory returns an instance of this Actor type (or a subclass), and based on
  * which Translate and Rotate UProperties are initialized, will associate those Components
  * with UInteractiveGizmo's that implement Axis Translation, Plane Translation, and Axis Rotation.
  * 
@@ -36,12 +36,12 @@ class FTransformGizmoTransformChange;
  * Actor suitable for use in a standard 3-axis Transformation Gizmo.
  */
 UCLASS(Transient)
-class INTERACTIVETOOLSFRAMEWORK_API ATransformGizmoActor : public AGizmoActor
+class INTERACTIVETOOLSFRAMEWORK_API ACombinedTransformGizmoActor : public AGizmoActor
 {
 	GENERATED_BODY()
 public:
 
-	ATransformGizmoActor();
+	ACombinedTransformGizmoActor();
 
 public:
 	//
@@ -127,18 +127,18 @@ public:
 
 public:
 	/**
-	 * Create a new instance of ATransformGizmoActor and populate the various
+	 * Create a new instance of ACombinedTransformGizmoActor and populate the various
 	 * sub-components with standard GizmoXComponent instances suitable for a 3-axis transformer Gizmo
 	 */
-	static ATransformGizmoActor* ConstructDefault3AxisGizmo(
+	static ACombinedTransformGizmoActor* ConstructDefault3AxisGizmo(
 		UWorld* World, UGizmoViewContext* GizmoViewContext
 	);
 
 	/**
-	 * Create a new instance of ATransformGizmoActor. Populate the sub-components 
+	 * Create a new instance of ACombinedTransformGizmoActor. Populate the sub-components 
 	 * specified by Elements with standard GizmoXComponent instances suitable for a 3-axis transformer Gizmo
 	 */
-	static ATransformGizmoActor* ConstructCustom3AxisGizmo(
+	static ACombinedTransformGizmoActor* ConstructCustom3AxisGizmo(
 		UWorld* World, UGizmoViewContext* GizmoViewContext,
 		ETransformGizmoSubElements Elements
 	);
@@ -149,25 +149,25 @@ public:
 
 
 /**
- * FTransformGizmoActorFactory creates new instances of ATransformGizmoActor which
- * are used by UTransformGizmo to implement 3D transformation Gizmos. 
- * An instance of FTransformGizmoActorFactory is passed to UTransformGizmo
- * (by way of UTransformGizmoBuilder), which then calls CreateNewGizmoActor()
+ * FCombinedTransformGizmoActorFactory creates new instances of ACombinedTransformGizmoActor which
+ * are used by UCombinedTransformGizmo to implement 3D transformation Gizmos. 
+ * An instance of FCombinedTransformGizmoActorFactory is passed to UCombinedTransformGizmo
+ * (by way of UCombinedTransformGizmoBuilder), which then calls CreateNewGizmoActor()
  * to spawn new Gizmo Actors.
  * 
  * By default CreateNewGizmoActor() returns a default Gizmo Actor suitable for
  * a three-axis transformation Gizmo, override this function to customize
  * the Actor sub-elements.
  */
-class INTERACTIVETOOLSFRAMEWORK_API FTransformGizmoActorFactory
+class INTERACTIVETOOLSFRAMEWORK_API FCombinedTransformGizmoActorFactory
 {
 public:
-	FTransformGizmoActorFactory(UGizmoViewContext* GizmoViewContextIn)
+	FCombinedTransformGizmoActorFactory(UGizmoViewContext* GizmoViewContextIn)
 		: GizmoViewContext(GizmoViewContextIn)
 	{
 	}
 
-	/** Only these members of the ATransformGizmoActor gizmo will be initialized */
+	/** Only these members of the ACombinedTransformGizmoActor gizmo will be initialized */
 	ETransformGizmoSubElements EnableElements =
 		ETransformGizmoSubElements::TranslateAllAxes |
 		ETransformGizmoSubElements::TranslateAllPlanes |
@@ -178,9 +178,9 @@ public:
 
 	/**
 	 * @param World the UWorld to create the new Actor in
-	 * @return new ATransformGizmoActor instance with members initialized with Components suitable for a transformation Gizmo
+	 * @return new ACombinedTransformGizmoActor instance with members initialized with Components suitable for a transformation Gizmo
 	 */
-	virtual ATransformGizmoActor* CreateNewGizmoActor(UWorld* World) const;
+	virtual ACombinedTransformGizmoActor* CreateNewGizmoActor(UWorld* World) const;
 
 protected:
 	/**
@@ -197,7 +197,7 @@ protected:
 
 
 UCLASS()
-class INTERACTIVETOOLSFRAMEWORK_API UTransformGizmoBuilder : public UInteractiveGizmoBuilder
+class INTERACTIVETOOLSFRAMEWORK_API UCombinedTransformGizmoBuilder : public UInteractiveGizmoBuilder
 {
 	GENERATED_BODY()
 
@@ -212,15 +212,15 @@ public:
 	FString AxisAngleBuilderIdentifier;
 
 	/**
-	 * If set, this Actor Builder will be passed to UTransformGizmo instances.
-	 * Otherwise new instances of the base FTransformGizmoActorFactory are created internally.
+	 * If set, this Actor Builder will be passed to UCombinedTransformGizmo instances.
+	 * Otherwise new instances of the base FCombinedTransformGizmoActorFactory are created internally.
 	 */
-	TSharedPtr<FTransformGizmoActorFactory> GizmoActorBuilder;
+	TSharedPtr<FCombinedTransformGizmoActorFactory> GizmoActorBuilder;
 
 	/**
-	 * If set, this hover function will be passed to UTransformGizmo instances to use instead of the default.
-	 * Hover is complicated for UTransformGizmo because all it knows about the different gizmo scene elements
-	 * is that they are UPrimitiveComponent (coming from the ATransformGizmoActor). The default hover
+	 * If set, this hover function will be passed to UCombinedTransformGizmo instances to use instead of the default.
+	 * Hover is complicated for UCombinedTransformGizmo because all it knows about the different gizmo scene elements
+	 * is that they are UPrimitiveComponent (coming from the ACombinedTransformGizmoActor). The default hover
 	 * function implementation is to try casting to UGizmoBaseComponent and calling ::UpdateHoverState().
 	 * If you are using different Components that do not subclass UGizmoBaseComponent, and you want hover to 
 	 * work, you will need to provide a different hover update function.
@@ -228,8 +228,8 @@ public:
 	TFunction<void(UPrimitiveComponent*, bool)> UpdateHoverFunction;
 
 	/**
-	 * If set, this coord-system function will be passed to UTransformGizmo instances to use instead
-	 * of the default UpdateCoordSystemFunction. By default the UTransformGizmo will query the external Context
+	 * If set, this coord-system function will be passed to UCombinedTransformGizmo instances to use instead
+	 * of the default UpdateCoordSystemFunction. By default the UCombinedTransformGizmo will query the external Context
 	 * to ask whether it should be using world or local coordinate system. Then the default UpdateCoordSystemFunction
 	 * will try casting to UGizmoBaseCmponent and passing that info on via UpdateWorldLocalState();
 	 * If you are using different Components that do not subclass UGizmoBaseComponent, and you want the coord system
@@ -243,29 +243,29 @@ public:
 
 
 /**
- * UTransformGizmo provides standard Transformation Gizmo interactions,
+ * UCombinedTransformGizmo provides standard Transformation Gizmo interactions,
  * applied to a UTransformProxy target object. By default the Gizmo will be
  * a standard XYZ translate/rotate Gizmo (axis and plane translation).
  * 
- * The in-scene representation of the Gizmo is a ATransformGizmoActor (or subclass).
+ * The in-scene representation of the Gizmo is a ACombinedTransformGizmoActor (or subclass).
  * This Actor has FProperty members for the various sub-widgets, each as a separate Component.
  * Any particular sub-widget of the Gizmo can be disabled by setting the respective
  * Actor Component to null. 
  * 
  * So, to create non-standard variants of the Transform Gizmo, set a new GizmoActorBuilder 
- * in the UTransformGizmoBuilder registered with the GizmoManager. Return
+ * in the UCombinedTransformGizmoBuilder registered with the GizmoManager. Return
  * a suitably-configured GizmoActor and everything else will be handled automatically.
  * 
  */
 UCLASS()
-class INTERACTIVETOOLSFRAMEWORK_API UTransformGizmo : public UInteractiveGizmo
+class INTERACTIVETOOLSFRAMEWORK_API UCombinedTransformGizmo : public UInteractiveGizmo
 {
 	GENERATED_BODY()
 
 public:
 
 	virtual void SetWorld(UWorld* World);
-	virtual void SetGizmoActorBuilder(TSharedPtr<FTransformGizmoActorFactory> Builder);
+	virtual void SetGizmoActorBuilder(TSharedPtr<FCombinedTransformGizmoActorFactory> Builder);
 	virtual void SetSubGizmoBuilderIdentifiers(FString AxisPositionBuilderIdentifier, FString PlanePositionBuilderIdentifier, FString AxisAngleBuilderIdentifier);
 	virtual void SetUpdateHoverFunction(TFunction<void(UPrimitiveComponent*, bool)> HoverFunction);
 	virtual void SetUpdateCoordSystemFunction(TFunction<void(UPrimitiveComponent*, EToolContextCoordinateSystem)> CoordSysFunction);
@@ -326,7 +326,7 @@ public:
 	/**
 	 * @return the internal GizmoActor used by the Gizmo
 	 */
-	ATransformGizmoActor* GetGizmoActor() const { return GizmoActor; }
+	ACombinedTransformGizmoActor* GetGizmoActor() const { return GizmoActor; }
 
 	/**
 	 * @return current transform of Gizmo
@@ -410,7 +410,7 @@ public:
 
 
 protected:
-	TSharedPtr<FTransformGizmoActorFactory> GizmoActorBuilder;
+	TSharedPtr<FCombinedTransformGizmoActorFactory> GizmoActorBuilder;
 
 	FString AxisPositionBuilderIdentifier;
 	FString PlanePositionBuilderIdentifier;
@@ -446,7 +446,7 @@ protected:
 	UWorld* World;
 
 	/** Current active GizmoActor that was spawned by this Gizmo. Will be destroyed when Gizmo is. */
-	ATransformGizmoActor* GizmoActor;
+	ACombinedTransformGizmoActor* GizmoActor;
 
 	//
 	// Axis Sources
@@ -498,7 +498,7 @@ protected:
 
 	/** 
 	 * State target is shared across gizmos, and created internally during SetActiveTarget(). 
-	 * Several FChange providers are registered with this StateTarget, including the UTransformGizmo
+	 * Several FChange providers are registered with this StateTarget, including the UCombinedTransformGizmo
 	 * itself (IToolCommandChangeSource implementation above is called)
 	 */
 	UPROPERTY()
diff --git a/Engine/Source/Runtime/InteractiveToolsFramework/Public/BaseGizmos/TransformGizmoUtil.h b/Engine/Source/Runtime/InteractiveToolsFramework/Public/BaseGizmos/TransformGizmoUtil.h
--- a/Engine/Source/Runtime/InteractiveToolsFramework/Public/BaseGizmos/TransformGizmoUtil.h
+++ b/Engine/Source/Runtime/InteractiveToolsFramework/Public/BaseGizmos/TransformGizmoUtil.h
@@ -11,7 +11,7 @@
 class UInteractiveToolsContext;
 class UInteractiveToolManager;
 class UInteractiveGizmoManager;
-class FTransformGizmoActorFactory;
+class FCombinedTransformGizmoActorFactory;
 
 
 namespace UE
@@ -19,73 +19,73 @@ namespace UE
 namespace TransformGizmoUtil
 {
 	//
-	// The functions below are helper functions that simplify usage of a UTransformGizmoContextObject 
+	// The functions below are helper functions that simplify usage of a UCombinedTransformGizmoContextObject 
 	// that is registered as a ContextStoreObject in an InteractiveToolsContext
 	//
 
 	/**
-	 * If one does not already exist, create a new instance of UTransformGizmoContextObject and add it to the
+	 * If one does not already exist, create a new instance of UCombinedTransformGizmoContextObject and add it to the
 	 * ToolsContext's ContextObjectStore
-	 * @return true if the ContextObjectStore now has a UTransformGizmoContextObject (whether it already existed, or was created)
+	 * @return true if the ContextObjectStore now has a UCombinedTransformGizmoContextObject (whether it already existed, or was created)
 	 */
 	INTERACTIVETOOLSFRAMEWORK_API bool RegisterTransformGizmoContextObject(UInteractiveToolsContext* ToolsContext);
 
 	/**
-	 * Remove any existing UTransformGizmoContextObject from the ToolsContext's ContextObjectStore
-	 * @return true if the ContextObjectStore no longer has a UTransformGizmoContextObject (whether it was removed, or did not exist)
+	 * Remove any existing UCombinedTransformGizmoContextObject from the ToolsContext's ContextObjectStore
+	 * @return true if the ContextObjectStore no longer has a UCombinedTransformGizmoContextObject (whether it was removed, or did not exist)
 	 */
 	INTERACTIVETOOLSFRAMEWORK_API bool DeregisterTransformGizmoContextObject(UInteractiveToolsContext* ToolsContext);
 
 
 	/**
-	 * Spawn a new standard 3-axis Transform gizmo (see UTransformGizmoContextObject::Create3AxisTransformGizmo for details)
-	 * GizmoManager's ToolsContext must have a UTransformGizmoContextObject registered (see UTransformGizmoContextObject for details)
+	 * Spawn a new standard 3-axis Transform gizmo (see UCombinedTransformGizmoContextObject::Create3AxisTransformGizmo for details)
+	 * GizmoManager's ToolsContext must have a UCombinedTransformGizmoContextObject registered (see UCombinedTransformGizmoContextObject for details)
 	 */
-	INTERACTIVETOOLSFRAMEWORK_API UTransformGizmo* Create3AxisTransformGizmo(UInteractiveGizmoManager* GizmoManager, void* Owner = nullptr, const FString& InstanceIdentifier = FString());
+	INTERACTIVETOOLSFRAMEWORK_API UCombinedTransformGizmo* Create3AxisTransformGizmo(UInteractiveGizmoManager* GizmoManager, void* Owner = nullptr, const FString& InstanceIdentifier = FString());
 	/**
-	 * Spawn a new standard 3-axis Transform gizmo (see UTransformGizmoContextObject::Create3AxisTransformGizmo for details)
-	 * ToolManager's ToolsContext must have a UTransformGizmoContextObject registered (see UTransformGizmoContextObject for details)
+	 * Spawn a new standard 3-axis Transform gizmo (see UCombinedTransformGizmoContextObject::Create3AxisTransformGizmo for details)
+	 * ToolManager's ToolsContext must have a UCombinedTransformGizmoContextObject registered (see UCombinedTransformGizmoContextObject for details)
 	 */
-	INTERACTIVETOOLSFRAMEWORK_API UTransformGizmo* Create3AxisTransformGizmo(UInteractiveToolManager* ToolManager, void* Owner = nullptr, const FString& InstanceIdentifier = FString());
+	INTERACTIVETOOLSFRAMEWORK_API UCombinedTransformGizmo* Create3AxisTransformGizmo(UInteractiveToolManager* ToolManager, void* Owner = nullptr, const FString& InstanceIdentifier = FString());
 
 
 	/**
-	 * Spawn a new custom Transform gizmo (see UTransformGizmoContextObject::CreateCustomTransformGizmo for details)
-	 * GizmoManager's ToolsContext must have a UTransformGizmoContextObject registered (see UTransformGizmoContextObject for details)
+	 * Spawn a new custom Transform gizmo (see UCombinedTransformGizmoContextObject::CreateCustomTransformGizmo for details)
+	 * GizmoManager's ToolsContext must have a UCombinedTransformGizmoContextObject registered (see UCombinedTransformGizmoContextObject for details)
 	 */
-	INTERACTIVETOOLSFRAMEWORK_API UTransformGizmo* CreateCustomTransformGizmo(UInteractiveGizmoManager* GizmoManager, ETransformGizmoSubElements Elements, void* Owner = nullptr, const FString& InstanceIdentifier = FString());
+	INTERACTIVETOOLSFRAMEWORK_API UCombinedTransformGizmo* CreateCustomTransformGizmo(UInteractiveGizmoManager* GizmoManager, ETransformGizmoSubElements Elements, void* Owner = nullptr, const FString& InstanceIdentifier = FString());
 	/**
-	 * Spawn a new custom Transform gizmo (see UTransformGizmoContextObject::CreateCustomTransformGizmo for details)
-	 * ToolManager's ToolsContext must have a UTransformGizmoContextObject registered (see UTransformGizmoContextObject for details)
+	 * Spawn a new custom Transform gizmo (see UCombinedTransformGizmoContextObject::CreateCustomTransformGizmo for details)
+	 * ToolManager's ToolsContext must have a UCombinedTransformGizmoContextObject registered (see UCombinedTransformGizmoContextObject for details)
 	 */
-	INTERACTIVETOOLSFRAMEWORK_API UTransformGizmo* CreateCustomTransformGizmo(UInteractiveToolManager* ToolManager, ETransformGizmoSubElements Elements, void* Owner = nullptr, const FString& InstanceIdentifier = FString());
+	INTERACTIVETOOLSFRAMEWORK_API UCombinedTransformGizmo* CreateCustomTransformGizmo(UInteractiveToolManager* ToolManager, ETransformGizmoSubElements Elements, void* Owner = nullptr, const FString& InstanceIdentifier = FString());
 
 
 	/**
-	 * Spawn a new custom Transform gizmo (see UTransformGizmoContextObject::CreateCustomTransformGizmo for details)
-	 * GizmoManager's ToolsContext must have a UTransformGizmoContextObject registered (see UTransformGizmoContextObject for details)
+	 * Spawn a new custom Transform gizmo (see UCombinedTransformGizmoContextObject::CreateCustomTransformGizmo for details)
+	 * GizmoManager's ToolsContext must have a UCombinedTransformGizmoContextObject registered (see UCombinedTransformGizmoContextObject for details)
 	 */
-	INTERACTIVETOOLSFRAMEWORK_API UTransformGizmo* CreateCustomRepositionableTransformGizmo(UInteractiveGizmoManager* GizmoManager, ETransformGizmoSubElements Elements, void* Owner = nullptr, const FString& InstanceIdentifier = FString());
+	INTERACTIVETOOLSFRAMEWORK_API UCombinedTransformGizmo* CreateCustomRepositionableTransformGizmo(UInteractiveGizmoManager* GizmoManager, ETransformGizmoSubElements Elements, void* Owner = nullptr, const FString& InstanceIdentifier = FString());
 	/**
-	 * Spawn a new custom Transform gizmo (see UTransformGizmoContextObject::CreateCustomRepositionableTransformGizmo for details)
-	 * ToolManager's ToolsContext must have a UTransformGizmoContextObject registered (see UTransformGizmoContextObject for details)
+	 * Spawn a new custom Transform gizmo (see UCombinedTransformGizmoContextObject::CreateCustomRepositionableTransformGizmo for details)
+	 * ToolManager's ToolsContext must have a UCombinedTransformGizmoContextObject registered (see UCombinedTransformGizmoContextObject for details)
 	 */
-	INTERACTIVETOOLSFRAMEWORK_API UTransformGizmo* CreateCustomRepositionableTransformGizmo(UInteractiveToolManager* ToolManager, ETransformGizmoSubElements Elements, void* Owner = nullptr, const FString& InstanceIdentifier = FString());
+	INTERACTIVETOOLSFRAMEWORK_API UCombinedTransformGizmo* CreateCustomRepositionableTransformGizmo(UInteractiveToolManager* ToolManager, ETransformGizmoSubElements Elements, void* Owner = nullptr, const FString& InstanceIdentifier = FString());
 }
 }
 
 
 /**
- * UTransformGizmoContextObject is a utility object that registers a set of Gizmo Builders
- * for UTransformGizmo and variants. The intended usage is to call RegisterGizmosWithManager(),
- * and then the UTransformGizmoContextObject will register itself as a ContextObject in the
+ * UCombinedTransformGizmoContextObject is a utility object that registers a set of Gizmo Builders
+ * for UCombinedTransformGizmo and variants. The intended usage is to call RegisterGizmosWithManager(),
+ * and then the UCombinedTransformGizmoContextObject will register itself as a ContextObject in the
  * InteractiveToolsContext's ContextObjectStore. Then the Create3AxisTransformGizmo()/etc functions
- * will spawn different variants of UTransformGizmo. The above UE::TransformGizmoUtil:: functions
- * will look up the UTransformGizmoContextObject instance in the ContextObjectStore and then
+ * will spawn different variants of UCombinedTransformGizmo. The above UE::TransformGizmoUtil:: functions
+ * will look up the UCombinedTransformGizmoContextObject instance in the ContextObjectStore and then
  * call the associated function below.
  */
 UCLASS(Transient)
-class INTERACTIVETOOLSFRAMEWORK_API UTransformGizmoContextObject : public UObject
+class INTERACTIVETOOLSFRAMEWORK_API UCombinedTransformGizmoContextObject : public UObject
 {
 	GENERATED_BODY()
 public:
@@ -108,7 +108,7 @@ public:
 	 * @param InstanceIdentifier optional client-defined *unique* string that can be used to locate this instance
 	 * @return new Gizmo instance that has been created and initialized
 	 */
-	virtual UTransformGizmo* Create3AxisTransformGizmo(UInteractiveGizmoManager* GizmoManager, void* Owner = nullptr, const FString& InstanceIdentifier = FString());
+	virtual UCombinedTransformGizmo* Create3AxisTransformGizmo(UInteractiveGizmoManager* GizmoManager, void* Owner = nullptr, const FString& InstanceIdentifier = FString());
 
 	/**
 	 * Activate a new customized instance of the default 3-axis transformation Gizmo, with only certain elements included. RegisterDefaultGizmos() must have been called first.
@@ -117,15 +117,15 @@ public:
 	 * @param InstanceIdentifier optional client-defined *unique* string that can be used to locate this instance
 	 * @return new Gizmo instance that has been created and initialized
 	 */
-	virtual UTransformGizmo* CreateCustomTransformGizmo(UInteractiveGizmoManager* GizmoManager, ETransformGizmoSubElements Elements, void* Owner = nullptr, const FString& InstanceIdentifier = FString());
+	virtual UCombinedTransformGizmo* CreateCustomTransformGizmo(UInteractiveGizmoManager* GizmoManager, ETransformGizmoSubElements Elements, void* Owner = nullptr, const FString& InstanceIdentifier = FString());
 
 	/**
-	 * Variant of CreateCustomTransformGizmo that creates a URepositionableTransformGizmo, which is an extension to UTransformGizmo that 
+	 * Variant of CreateCustomTransformGizmo that creates a URepositionableTransformGizmo, which is an extension to UCombinedTransformGizmo that 
 	 * supports various snapping interactions
 	 */
-	virtual UTransformGizmo* CreateCustomRepositionableTransformGizmo(UInteractiveGizmoManager* GizmoManager, ETransformGizmoSubElements Elements, void* Owner = nullptr, const FString& InstanceIdentifier = FString());
+	virtual UCombinedTransformGizmo* CreateCustomRepositionableTransformGizmo(UInteractiveGizmoManager* GizmoManager, ETransformGizmoSubElements Elements, void* Owner = nullptr, const FString& InstanceIdentifier = FString());
 
 protected:
-	TSharedPtr<FTransformGizmoActorFactory> GizmoActorBuilder;
+	TSharedPtr<FCombinedTransformGizmoActorFactory> GizmoActorBuilder;
 	bool bDefaultGizmosRegistered = false;
 };
diff --git a/Engine/Source/Runtime/InteractiveToolsFramework/Public/InteractiveGizmo.h b/Engine/Source/Runtime/InteractiveToolsFramework/Public/InteractiveGizmo.h
--- a/Engine/Source/Runtime/InteractiveToolsFramework/Public/InteractiveGizmo.h
+++ b/Engine/Source/Runtime/InteractiveToolsFramework/Public/InteractiveGizmo.h
@@ -90,7 +90,7 @@ protected:
 
 /**
  * ETransformGizmoSubElements identifies the sub-elements of a standard 3-axis transformation Gizmo.
- * Used by GizmoManager to customize UTransformGizmo instances.
+ * Used by GizmoManager to customize UCombinedTransformGizmo instances.
  */
 UENUM()
 enum class ETransformGizmoSubElements
diff --git a/Engine/Source/Runtime/InteractiveToolsFramework/Public/InteractiveGizmoManager.h b/Engine/Source/Runtime/InteractiveToolsFramework/Public/InteractiveGizmoManager.h
--- a/Engine/Source/Runtime/InteractiveToolsFramework/Public/InteractiveGizmoManager.h
+++ b/Engine/Source/Runtime/InteractiveToolsFramework/Public/InteractiveGizmoManager.h
@@ -11,9 +11,9 @@
 #include "ToolContextInterfaces.h"
 #include "InteractiveGizmoManager.generated.h"
 
-class FTransformGizmoActorFactory;
-class UTransformGizmo;
-class UTransformGizmoBuilder;
+class FCombinedTransformGizmoActorFactory;
+class UCombinedTransformGizmo;
+class UCombinedTransformGizmoBuilder;
 class UContextObjectStore;
 
 USTRUCT()
@@ -206,7 +206,7 @@ public:
 	 * @param InstanceIdentifier optional client-defined *unique* string that can be used to locate this instance
 	 * @return new Gizmo instance that has been created and initialized
 	 */
-	virtual UTransformGizmo* Create3AxisTransformGizmo(void* Owner = nullptr, const FString& InstanceIdentifier = FString());
+	virtual UCombinedTransformGizmo* Create3AxisTransformGizmo(void* Owner = nullptr, const FString& InstanceIdentifier = FString());
 
 	/**
 	 * Activate a new customized instance of the default 3-axis transformation Gizmo, with only certain elements included. RegisterDefaultGizmos() must have been called first.
@@ -215,9 +215,9 @@ public:
 	 * @param InstanceIdentifier optional client-defined *unique* string that can be used to locate this instance
 	 * @return new Gizmo instance that has been created and initialized
 	 */
-	virtual UTransformGizmo* CreateCustomTransformGizmo(ETransformGizmoSubElements Elements, void* Owner = nullptr, const FString& InstanceIdentifier = FString());
+	virtual UCombinedTransformGizmo* CreateCustomTransformGizmo(ETransformGizmoSubElements Elements, void* Owner = nullptr, const FString& InstanceIdentifier = FString());
 
-	virtual UTransformGizmo* CreateCustomRepositionableTransformGizmo(ETransformGizmoSubElements Elements, void* Owner = nullptr, const FString& InstanceIdentifier = FString());
+	virtual UCombinedTransformGizmo* CreateCustomRepositionableTransformGizmo(ETransformGizmoSubElements Elements, void* Owner = nullptr, const FString& InstanceIdentifier = FString());
 
 public:
 	// builder identifiers for default gizmo types. Perhaps should have an API for this...
@@ -249,5 +249,5 @@ protected:
 	TMap<FString, TObjectPtr<UInteractiveGizmoBuilder>> GizmoBuilders;
 
 	bool bDefaultGizmosRegistered = false;
-	TSharedPtr<FTransformGizmoActorFactory> GizmoActorBuilder;
+	TSharedPtr<FCombinedTransformGizmoActorFactory> GizmoActorBuilder;
 };
