ToolsFramework: Rename TransformGizmo.h to CombinedTransformGizmo.h
#rnx

diff --git a/Engine/Plugins/Editor/SequencerAnimTools/Source/SequencerAnimTools/Private/MotionTrailTool.cpp b/Engine/Plugins/Editor/SequencerAnimTools/Source/SequencerAnimTools/Private/MotionTrailTool.cpp
--- a/Engine/Plugins/Editor/SequencerAnimTools/Source/SequencerAnimTools/Private/MotionTrailTool.cpp
+++ b/Engine/Plugins/Editor/SequencerAnimTools/Source/SequencerAnimTools/Private/MotionTrailTool.cpp
@@ -5,7 +5,7 @@
 #include "ISequencer.h"
 #include "Sequencer/SequencerTrailHierarchy.h"
 #include "Sequencer/MovieSceneTransformTrail.h"
-#include "BaseGizmos/TransformGizmo.h"
+#include "BaseGizmos/CombinedTransformGizmo.h"
 #include "InteractiveGizmoManager.h"
 #include "BaseBehaviors/SingleClickBehavior.h"
 #include "BaseBehaviors/SingleKeyCaptureBehavior.h"
diff --git a/Engine/Plugins/Editor/SequencerAnimTools/Source/SequencerAnimTools/Private/Sequencer/MovieSceneTransformTrail.cpp b/Engine/Plugins/Editor/SequencerAnimTools/Source/SequencerAnimTools/Private/Sequencer/MovieSceneTransformTrail.cpp
--- a/Engine/Plugins/Editor/SequencerAnimTools/Source/SequencerAnimTools/Private/Sequencer/MovieSceneTransformTrail.cpp
+++ b/Engine/Plugins/Editor/SequencerAnimTools/Source/SequencerAnimTools/Private/Sequencer/MovieSceneTransformTrail.cpp
@@ -14,7 +14,7 @@
 #include "Sequencer/MovieSceneControlRigParameterSection.h"
 
 #include "BaseGizmos/GizmoComponents.h"
-#include "BaseGizmos/TransformGizmo.h"
+#include "BaseGizmos/CombinedTransformGizmo.h"
 
 #include "ViewportWorldInteraction.h"
 #include "GameFramework/Actor.h"
diff --git a/Engine/Plugins/Editor/SequencerAnimTools/Source/SequencerAnimTools/Private/SequencerAnimEditPivotTool.cpp b/Engine/Plugins/Editor/SequencerAnimTools/Source/SequencerAnimTools/Private/SequencerAnimEditPivotTool.cpp
--- a/Engine/Plugins/Editor/SequencerAnimTools/Source/SequencerAnimTools/Private/SequencerAnimEditPivotTool.cpp
+++ b/Engine/Plugins/Editor/SequencerAnimTools/Source/SequencerAnimTools/Private/SequencerAnimEditPivotTool.cpp
@@ -5,7 +5,7 @@
 #include "ISequencer.h"
 #include "Sequencer/SequencerTrailHierarchy.h"
 #include "Sequencer/MovieSceneTransformTrail.h"
-#include "BaseGizmos/TransformGizmo.h"
+#include "BaseGizmos/CombinedTransformGizmo.h"
 #include "InteractiveGizmoManager.h"
 #include "BaseBehaviors/SingleClickBehavior.h"
 #include "BaseBehaviors/SingleKeyCaptureBehavior.h"
diff --git a/Engine/Plugins/Experimental/ChaosEditor/Source/FractureEditor/Private/FractureToolCutter.h b/Engine/Plugins/Experimental/ChaosEditor/Source/FractureEditor/Private/FractureToolCutter.h
--- a/Engine/Plugins/Experimental/ChaosEditor/Source/FractureEditor/Private/FractureToolCutter.h
+++ b/Engine/Plugins/Experimental/ChaosEditor/Source/FractureEditor/Private/FractureToolCutter.h
@@ -4,7 +4,7 @@
 
 #include "FractureTool.h"
 
-#include "BaseGizmos/TransformGizmo.h"
+#include "BaseGizmos/CombinedTransformGizmo.h"
 #include "BaseGizmos/TransformProxy.h"
 #include "Drawing/LineSetComponent.h"
 
diff --git a/Engine/Plugins/Experimental/GizmoEdMode/Source/GizmoEdMode/Private/DefaultAssetEditorGizmoFactory.cpp b/Engine/Plugins/Experimental/GizmoEdMode/Source/GizmoEdMode/Private/DefaultAssetEditorGizmoFactory.cpp
--- a/Engine/Plugins/Experimental/GizmoEdMode/Source/GizmoEdMode/Private/DefaultAssetEditorGizmoFactory.cpp
+++ b/Engine/Plugins/Experimental/GizmoEdMode/Source/GizmoEdMode/Private/DefaultAssetEditorGizmoFactory.cpp
@@ -3,7 +3,7 @@
 #include "DefaultAssetEditorGizmoFactory.h"
 #include "EditorModeManager.h"
 #include "Engine/Selection.h"
-#include "BaseGizmos/TransformGizmo.h"
+#include "BaseGizmos/CombinedTransformGizmo.h"
 #include "BaseGizmos/TransformProxy.h"
 #include "UnrealWidgetFwd.h"
 
diff --git a/Engine/Plugins/Experimental/GizmoEdMode/Source/GizmoEdMode/Private/GizmoEdMode.cpp b/Engine/Plugins/Experimental/GizmoEdMode/Source/GizmoEdMode/Private/GizmoEdMode.cpp
--- a/Engine/Plugins/Experimental/GizmoEdMode/Source/GizmoEdMode/Private/GizmoEdMode.cpp
+++ b/Engine/Plugins/Experimental/GizmoEdMode/Source/GizmoEdMode/Private/GizmoEdMode.cpp
@@ -6,7 +6,7 @@
 #include "EditorModeManager.h"
 #include "EdModeInteractiveToolsContext.h"
 #include "InteractiveGizmo.h"
-#include "BaseGizmos/TransformGizmo.h"
+#include "BaseGizmos/CombinedTransformGizmo.h"
 #include "UnrealWidgetFwd.h"
 #include "Utils.h"
 #include "Stats/Stats2.h"
diff --git a/Engine/Plugins/Experimental/GizmoEdMode/Source/LightGizmos/Private/DirectionalLightGizmo.h b/Engine/Plugins/Experimental/GizmoEdMode/Source/LightGizmos/Private/DirectionalLightGizmo.h
--- a/Engine/Plugins/Experimental/GizmoEdMode/Source/LightGizmos/Private/DirectionalLightGizmo.h
+++ b/Engine/Plugins/Experimental/GizmoEdMode/Source/LightGizmos/Private/DirectionalLightGizmo.h
@@ -10,7 +10,7 @@
 #include "BaseBehaviors/BehaviorTargetInterfaces.h"
 #include "BaseBehaviors/AnyButtonInputBehavior.h"
 #include "Engine/DirectionalLight.h"
-#include "BaseGizmos/TransformGizmo.h"
+#include "BaseGizmos/CombinedTransformGizmo.h"
 #include "SubTransformProxy.h"
 #include "BaseGizmos/GizmoBaseComponent.h"
 
diff --git a/Engine/Plugins/Experimental/GizmoEdMode/Source/LightGizmos/Private/DirectionalLightGizmoFactory.cpp b/Engine/Plugins/Experimental/GizmoEdMode/Source/LightGizmos/Private/DirectionalLightGizmoFactory.cpp
--- a/Engine/Plugins/Experimental/GizmoEdMode/Source/LightGizmos/Private/DirectionalLightGizmoFactory.cpp
+++ b/Engine/Plugins/Experimental/GizmoEdMode/Source/LightGizmos/Private/DirectionalLightGizmoFactory.cpp
@@ -3,7 +3,7 @@
 #include "DirectionalLightGizmoFactory.h"
 #include "EditorModeManager.h"
 #include "Engine/Selection.h"
-#include "BaseGizmos/TransformGizmo.h"
+#include "BaseGizmos/CombinedTransformGizmo.h"
 #include "BaseGizmos/TransformProxy.h"
 #include "Math/Rotator.h"
 #include "UnrealWidgetFwd.h"
diff --git a/Engine/Plugins/Experimental/GizmoEdMode/Source/LightGizmos/Private/PointLightGizmoFactory.cpp b/Engine/Plugins/Experimental/GizmoEdMode/Source/LightGizmos/Private/PointLightGizmoFactory.cpp
--- a/Engine/Plugins/Experimental/GizmoEdMode/Source/LightGizmos/Private/PointLightGizmoFactory.cpp
+++ b/Engine/Plugins/Experimental/GizmoEdMode/Source/LightGizmos/Private/PointLightGizmoFactory.cpp
@@ -3,7 +3,7 @@
 #include "PointLightGizmoFactory.h"
 #include "EditorModeManager.h"
 #include "Engine/Selection.h"
-#include "BaseGizmos/TransformGizmo.h"
+#include "BaseGizmos/CombinedTransformGizmo.h"
 #include "SubTransformProxy.h"
 #include "Math/Rotator.h"
 #include "UnrealWidgetFwd.h"
diff --git a/Engine/Plugins/Experimental/GizmoEdMode/Source/LightGizmos/Private/SpotLightGizmoFactory.cpp b/Engine/Plugins/Experimental/GizmoEdMode/Source/LightGizmos/Private/SpotLightGizmoFactory.cpp
--- a/Engine/Plugins/Experimental/GizmoEdMode/Source/LightGizmos/Private/SpotLightGizmoFactory.cpp
+++ b/Engine/Plugins/Experimental/GizmoEdMode/Source/LightGizmos/Private/SpotLightGizmoFactory.cpp
@@ -3,7 +3,7 @@
 #include "SpotLightGizmoFactory.h"
 #include "EditorModeManager.h"
 #include "Engine/Selection.h"
-#include "BaseGizmos/TransformGizmo.h"
+#include "BaseGizmos/CombinedTransformGizmo.h"
 #include "SubTransformProxy.h"
 #include "Math/Rotator.h"
 #include "UnrealWidgetFwd.h"
diff --git a/Engine/Plugins/Experimental/MeshModelingToolsetExp/Source/MeshModelingToolsEditorOnly/Private/AddPivotActorTool.cpp b/Engine/Plugins/Experimental/MeshModelingToolsetExp/Source/MeshModelingToolsEditorOnly/Private/AddPivotActorTool.cpp
--- a/Engine/Plugins/Experimental/MeshModelingToolsetExp/Source/MeshModelingToolsEditorOnly/Private/AddPivotActorTool.cpp
+++ b/Engine/Plugins/Experimental/MeshModelingToolsetExp/Source/MeshModelingToolsEditorOnly/Private/AddPivotActorTool.cpp
@@ -3,7 +3,7 @@
 #include "AddPivotActorTool.h"
 
 #include "ActorFactories/ActorFactoryEmptyActor.h"
-#include "BaseGizmos/TransformGizmo.h"
+#include "BaseGizmos/CombinedTransformGizmo.h"
 #include "BaseGizmos/TransformGizmoUtil.h"
 #include "BaseGizmos/TransformProxy.h"
 #include "Components/BillboardComponent.h"
diff --git a/Engine/Plugins/Experimental/MeshModelingToolsetExp/Source/MeshModelingToolsExp/Private/EditMeshPolygonsTool.cpp b/Engine/Plugins/Experimental/MeshModelingToolsetExp/Source/MeshModelingToolsExp/Private/EditMeshPolygonsTool.cpp
--- a/Engine/Plugins/Experimental/MeshModelingToolsetExp/Source/MeshModelingToolsExp/Private/EditMeshPolygonsTool.cpp
+++ b/Engine/Plugins/Experimental/MeshModelingToolsetExp/Source/MeshModelingToolsExp/Private/EditMeshPolygonsTool.cpp
@@ -4,7 +4,7 @@
 
 #include "Algo/ForEach.h"
 #include "BaseBehaviors/SingleClickBehavior.h"
-#include "BaseGizmos/TransformGizmo.h"
+#include "BaseGizmos/CombinedTransformGizmo.h"
 #include "BaseGizmos/TransformGizmoUtil.h"
 #include "BaseGizmos/TransformProxy.h"
 #include "CompGeom/PolygonTriangulation.h"
diff --git a/Engine/Plugins/Experimental/MeshModelingToolsetExp/Source/MeshModelingToolsExp/Private/PlaneCutTool.cpp b/Engine/Plugins/Experimental/MeshModelingToolsetExp/Source/MeshModelingToolsExp/Private/PlaneCutTool.cpp
--- a/Engine/Plugins/Experimental/MeshModelingToolsetExp/Source/MeshModelingToolsExp/Private/PlaneCutTool.cpp
+++ b/Engine/Plugins/Experimental/MeshModelingToolsetExp/Source/MeshModelingToolsExp/Private/PlaneCutTool.cpp
@@ -17,7 +17,7 @@
 #include "InteractiveGizmoManager.h"
 
 #include "BaseGizmos/GizmoComponents.h"
-#include "BaseGizmos/TransformGizmo.h"
+#include "BaseGizmos/CombinedTransformGizmo.h"
 
 #include "Drawing/MeshDebugDrawing.h"
 #include "ModelingObjectsCreationAPI.h"
diff --git a/Engine/Plugins/Experimental/MeshModelingToolsetExp/Source/MeshModelingToolsExp/Private/VoxelBlendMeshesTool.cpp b/Engine/Plugins/Experimental/MeshModelingToolsetExp/Source/MeshModelingToolsExp/Private/VoxelBlendMeshesTool.cpp
--- a/Engine/Plugins/Experimental/MeshModelingToolsetExp/Source/MeshModelingToolsExp/Private/VoxelBlendMeshesTool.cpp
+++ b/Engine/Plugins/Experimental/MeshModelingToolsetExp/Source/MeshModelingToolsExp/Private/VoxelBlendMeshesTool.cpp
@@ -18,7 +18,7 @@
 #include "InteractiveGizmoManager.h"
 
 #include "BaseGizmos/GizmoComponents.h"
-#include "BaseGizmos/TransformGizmo.h"
+#include "BaseGizmos/CombinedTransformGizmo.h"
 
 #include "CompositionOps/VoxelBlendMeshesOp.h"
 
diff --git a/Engine/Plugins/Experimental/MeshModelingToolsetExp/Source/MeshModelingToolsExp/Private/VoxelMorphologyMeshesTool.cpp b/Engine/Plugins/Experimental/MeshModelingToolsetExp/Source/MeshModelingToolsExp/Private/VoxelMorphologyMeshesTool.cpp
--- a/Engine/Plugins/Experimental/MeshModelingToolsetExp/Source/MeshModelingToolsExp/Private/VoxelMorphologyMeshesTool.cpp
+++ b/Engine/Plugins/Experimental/MeshModelingToolsetExp/Source/MeshModelingToolsExp/Private/VoxelMorphologyMeshesTool.cpp
@@ -18,7 +18,7 @@
 #include "InteractiveGizmoManager.h"
 
 #include "BaseGizmos/GizmoComponents.h"
-#include "BaseGizmos/TransformGizmo.h"
+#include "BaseGizmos/CombinedTransformGizmo.h"
 
 #include "CompositionOps/VoxelMorphologyMeshesOp.h"
 
diff --git a/Engine/Plugins/Experimental/MeshModelingToolsetExp/Source/MeshModelingToolsExp/Private/VoxelSolidifyMeshesTool.cpp b/Engine/Plugins/Experimental/MeshModelingToolsetExp/Source/MeshModelingToolsExp/Private/VoxelSolidifyMeshesTool.cpp
--- a/Engine/Plugins/Experimental/MeshModelingToolsetExp/Source/MeshModelingToolsExp/Private/VoxelSolidifyMeshesTool.cpp
+++ b/Engine/Plugins/Experimental/MeshModelingToolsetExp/Source/MeshModelingToolsExp/Private/VoxelSolidifyMeshesTool.cpp
@@ -18,7 +18,7 @@
 #include "InteractiveGizmoManager.h"
 
 #include "BaseGizmos/GizmoComponents.h"
-#include "BaseGizmos/TransformGizmo.h"
+#include "BaseGizmos/CombinedTransformGizmo.h"
 
 #include "CompositionOps/VoxelSolidifyMeshesOp.h"
 
diff --git a/Engine/Plugins/Experimental/MeshModelingToolsetExp/Source/MeshModelingToolsExp/Public/MirrorTool.h b/Engine/Plugins/Experimental/MeshModelingToolsetExp/Source/MeshModelingToolsExp/Public/MirrorTool.h
--- a/Engine/Plugins/Experimental/MeshModelingToolsetExp/Source/MeshModelingToolsExp/Public/MirrorTool.h
+++ b/Engine/Plugins/Experimental/MeshModelingToolsetExp/Source/MeshModelingToolsExp/Public/MirrorTool.h
@@ -4,7 +4,7 @@
 
 #include "CoreMinimal.h"
 
-#include "BaseGizmos/TransformGizmo.h"
+#include "BaseGizmos/CombinedTransformGizmo.h"
 #include "Changes/DynamicMeshChangeTarget.h"
 #include "InteractiveToolBuilder.h"
 #include "Mechanics/ConstructionPlaneMechanic.h"
diff --git a/Engine/Plugins/Experimental/MeshModelingToolsetExp/Source/ModelingComponents/Private/Mechanics/DragAlignmentMechanic.cpp b/Engine/Plugins/Experimental/MeshModelingToolsetExp/Source/ModelingComponents/Private/Mechanics/DragAlignmentMechanic.cpp
--- a/Engine/Plugins/Experimental/MeshModelingToolsetExp/Source/ModelingComponents/Private/Mechanics/DragAlignmentMechanic.cpp
+++ b/Engine/Plugins/Experimental/MeshModelingToolsetExp/Source/ModelingComponents/Private/Mechanics/DragAlignmentMechanic.cpp
@@ -5,7 +5,7 @@
 #include "BaseBehaviors/KeyAsModifierInputBehavior.h"
 #include "BaseGizmos/IntervalGizmo.h"
 #include "BaseGizmos/RepositionableTransformGizmo.h"
-#include "BaseGizmos/TransformGizmo.h"
+#include "BaseGizmos/CombinedTransformGizmo.h"
 #include "BaseGizmos/TransformProxy.h"
 #include "InputCoreTypes.h" //EKeys
 #include "Operations/GroupTopologyDeformer.h"
diff --git a/Engine/Plugins/Experimental/MotionTrailEditorMode/Source/MotionTrailEditorMode/Private/MotionTrailEditorToolset.cpp b/Engine/Plugins/Experimental/MotionTrailEditorMode/Source/MotionTrailEditorMode/Private/MotionTrailEditorToolset.cpp
--- a/Engine/Plugins/Experimental/MotionTrailEditorMode/Source/MotionTrailEditorMode/Private/MotionTrailEditorToolset.cpp
+++ b/Engine/Plugins/Experimental/MotionTrailEditorMode/Source/MotionTrailEditorMode/Private/MotionTrailEditorToolset.cpp
@@ -3,7 +3,7 @@
 #include "MotionTrailEditorToolset.h"
 #include "MotionTrailEditorMode.h"
 
-#include "BaseGizmos/TransformGizmo.h"
+#include "BaseGizmos/CombinedTransformGizmo.h"
 #include "InteractiveGizmoManager.h"
 
 #include "BaseBehaviors/SingleClickBehavior.h"
diff --git a/Engine/Plugins/Experimental/MotionTrailEditorMode/Source/MotionTrailEditorMode/Private/Sequencer/MovieSceneTransformTrail.cpp b/Engine/Plugins/Experimental/MotionTrailEditorMode/Source/MotionTrailEditorMode/Private/Sequencer/MovieSceneTransformTrail.cpp
--- a/Engine/Plugins/Experimental/MotionTrailEditorMode/Source/MotionTrailEditorMode/Private/Sequencer/MovieSceneTransformTrail.cpp
+++ b/Engine/Plugins/Experimental/MotionTrailEditorMode/Source/MotionTrailEditorMode/Private/Sequencer/MovieSceneTransformTrail.cpp
@@ -16,7 +16,7 @@
 #include "Sequencer/MovieSceneControlRigParameterSection.h"
 
 #include "BaseGizmos/GizmoComponents.h"
-#include "BaseGizmos/TransformGizmo.h"
+#include "BaseGizmos/CombinedTransformGizmo.h"
 
 #include "ViewportWorldInteraction.h"
 #include "GameFramework/Actor.h"
diff --git a/Engine/Plugins/Experimental/MotionTrailEditorMode/Source/MotionTrailEditorMode/Private/Sequencer/MovieSceneTransformTrailTool.cpp b/Engine/Plugins/Experimental/MotionTrailEditorMode/Source/MotionTrailEditorMode/Private/Sequencer/MovieSceneTransformTrailTool.cpp
--- a/Engine/Plugins/Experimental/MotionTrailEditorMode/Source/MotionTrailEditorMode/Private/Sequencer/MovieSceneTransformTrailTool.cpp
+++ b/Engine/Plugins/Experimental/MotionTrailEditorMode/Source/MotionTrailEditorMode/Private/Sequencer/MovieSceneTransformTrailTool.cpp
@@ -15,7 +15,7 @@
 #include "MovieSceneSequence.h"
 
 #include "BaseGizmos/GizmoComponents.h"
-#include "BaseGizmos/TransformGizmo.h"
+#include "BaseGizmos/CombinedTransformGizmo.h"
 
 #include "ViewportWorldInteraction.h"
 #include "GameFramework/Actor.h"
diff --git a/Engine/Plugins/Experimental/UVEditor/Source/UVEditorTools/Private/UVSelectTool.cpp b/Engine/Plugins/Experimental/UVEditor/Source/UVEditorTools/Private/UVSelectTool.cpp
--- a/Engine/Plugins/Experimental/UVEditor/Source/UVEditorTools/Private/UVSelectTool.cpp
+++ b/Engine/Plugins/Experimental/UVEditor/Source/UVEditorTools/Private/UVSelectTool.cpp
@@ -3,7 +3,7 @@
 #include "UVSelectTool.h"
 
 #include "Algo/Unique.h"
-#include "BaseGizmos/TransformGizmo.h"
+#include "BaseGizmos/CombinedTransformGizmo.h"
 #include "ContextObjectStore.h"
 #include "DynamicMesh/DynamicMeshChangeTracker.h"
 #include "Drawing/LineSetComponent.h"
diff --git a/Engine/Source/Runtime/InteractiveToolsFramework/Private/BaseGizmos/CombinedTransformGizmo.cpp b/Engine/Source/Runtime/InteractiveToolsFramework/Private/BaseGizmos/CombinedTransformGizmo.cpp
new file mode 100644
--- /dev/null
+++ b/Engine/Source/Runtime/InteractiveToolsFramework/Private/BaseGizmos/CombinedTransformGizmo.cpp
@@ -0,0 +1,978 @@
+// Copyright Epic Games, Inc. All Rights Reserved.
+
+#include "BaseGizmos/CombinedTransformGizmo.h"
+#include "InteractiveGizmoManager.h"
+#include "BaseGizmos/AxisPositionGizmo.h"
+#include "BaseGizmos/PlanePositionGizmo.h"
+#include "BaseGizmos/AxisAngleGizmo.h"
+#include "BaseGizmos/GizmoComponents.h"
+
+#include "BaseGizmos/GizmoArrowComponent.h"
+#include "BaseGizmos/GizmoRectangleComponent.h"
+#include "BaseGizmos/GizmoCircleComponent.h"
+#include "BaseGizmos/GizmoBoxComponent.h"
+#include "BaseGizmos/GizmoLineHandleComponent.h"
+#include "BaseGizmos/GizmoViewContext.h"
+
+// need this to implement hover
+#include "BaseGizmos/GizmoBaseComponent.h"
+
+#include "Components/SphereComponent.h"
+#include "Components/PrimitiveComponent.h"
+#include "ContextObjectStore.h"
+#include "Engine/World.h"
+#include "Engine/CollisionProfile.h"
+
+
+#define LOCTEXT_NAMESPACE "UCombinedTransformGizmo"
+
+
+ACombinedTransformGizmoActor::ACombinedTransformGizmoActor()
+{
+	// root component is a hidden sphere
+	USphereComponent* SphereComponent = CreateDefaultSubobject<USphereComponent>(TEXT("GizmoCenter"));
+	RootComponent = SphereComponent;
+	SphereComponent->InitSphereRadius(1.0f);
+	SphereComponent->SetVisibility(false);
+	SphereComponent->SetCollisionProfileName(UCollisionProfile::NoCollision_ProfileName);
+}
+
+
+
+ACombinedTransformGizmoActor* ACombinedTransformGizmoActor::ConstructDefault3AxisGizmo(UWorld* World, UGizmoViewContext* GizmoViewContext)
+{
+	return ConstructCustom3AxisGizmo(World, GizmoViewContext,
+		ETransformGizmoSubElements::TranslateAllAxes |
+		ETransformGizmoSubElements::TranslateAllPlanes |
+		ETransformGizmoSubElements::RotateAllAxes |
+		ETransformGizmoSubElements::ScaleAllAxes |
+		ETransformGizmoSubElements::ScaleAllPlanes |
+		ETransformGizmoSubElements::ScaleUniform
+	);
+}
+
+
+ACombinedTransformGizmoActor* ACombinedTransformGizmoActor::ConstructCustom3AxisGizmo(
+	UWorld* World, UGizmoViewContext* GizmoViewContext,
+	ETransformGizmoSubElements Elements)
+{
+	FActorSpawnParameters SpawnInfo;
+	ACombinedTransformGizmoActor* NewActor = World->SpawnActor<ACombinedTransformGizmoActor>(FVector::ZeroVector, FRotator::ZeroRotator, SpawnInfo);
+
+	float GizmoLineThickness = 3.0f;
+
+
+	auto MakeAxisArrowFunc = [&](const FLinearColor& Color, const FVector& Axis)
+	{
+		UGizmoArrowComponent* Component = AddDefaultArrowComponent(World, NewActor, GizmoViewContext, Color, Axis, 60.0f);
+		Component->Gap = 20.0f;
+		Component->Thickness = GizmoLineThickness;
+		Component->NotifyExternalPropertyUpdates();
+		return Component;
+	};
+	if ((Elements & ETransformGizmoSubElements::TranslateAxisX) != ETransformGizmoSubElements::None)
+	{
+		NewActor->TranslateX = MakeAxisArrowFunc(FLinearColor::Red, FVector(1, 0, 0));
+	}
+	if ((Elements & ETransformGizmoSubElements::TranslateAxisY) != ETransformGizmoSubElements::None)
+	{
+		NewActor->TranslateY = MakeAxisArrowFunc(FLinearColor::Green, FVector(0, 1, 0));
+	}
+	if ((Elements & ETransformGizmoSubElements::TranslateAxisZ) != ETransformGizmoSubElements::None)
+	{
+		NewActor->TranslateZ = MakeAxisArrowFunc(FLinearColor::Blue, FVector(0, 0, 1));
+	}
+
+
+	auto MakePlaneRectFunc = [&](const FLinearColor& Color, const FVector& Axis0, const FVector& Axis1)
+	{
+		UGizmoRectangleComponent* Component = AddDefaultRectangleComponent(World, NewActor, GizmoViewContext, Color, Axis0, Axis1);
+		Component->LengthX = Component->LengthY = 30.0f;
+		Component->SegmentFlags = 0x2 | 0x4;
+		Component->Thickness = GizmoLineThickness;
+		Component->NotifyExternalPropertyUpdates();
+		return Component;
+	};
+	if ((Elements & ETransformGizmoSubElements::TranslatePlaneYZ) != ETransformGizmoSubElements::None)
+	{
+		NewActor->TranslateYZ = MakePlaneRectFunc(FLinearColor::Red, FVector(0, 1, 0), FVector(0, 0, 1));
+	}
+	if ((Elements & ETransformGizmoSubElements::TranslatePlaneXZ) != ETransformGizmoSubElements::None)
+	{
+		NewActor->TranslateXZ = MakePlaneRectFunc(FLinearColor::Green, FVector(1, 0, 0), FVector(0, 0, 1));
+	}
+	if ((Elements & ETransformGizmoSubElements::TranslatePlaneXY) != ETransformGizmoSubElements::None)
+	{
+		NewActor->TranslateXY = MakePlaneRectFunc(FLinearColor::Blue, FVector(1, 0, 0), FVector(0, 1, 0));
+	}
+
+	auto MakeAxisRotateCircleFunc = [&](const FLinearColor& Color, const FVector& Axis)
+	{
+		UGizmoCircleComponent* Component = AddDefaultCircleComponent(World, NewActor, GizmoViewContext, Color, Axis, 120.0f);
+		Component->Thickness = GizmoLineThickness;
+		Component->NotifyExternalPropertyUpdates();
+		return Component;
+	};
+
+	bool bAnyRotate = false;
+	if ((Elements & ETransformGizmoSubElements::RotateAxisX) != ETransformGizmoSubElements::None)
+	{
+		NewActor->RotateX = MakeAxisRotateCircleFunc(FLinearColor::Red, FVector(1, 0, 0));
+		bAnyRotate = true;
+	}
+	if ((Elements & ETransformGizmoSubElements::RotateAxisY) != ETransformGizmoSubElements::None)
+	{
+		NewActor->RotateY = MakeAxisRotateCircleFunc(FLinearColor::Green, FVector(0, 1, 0));
+		bAnyRotate = true;
+	}
+	if ((Elements & ETransformGizmoSubElements::RotateAxisZ) != ETransformGizmoSubElements::None)
+	{
+		NewActor->RotateZ = MakeAxisRotateCircleFunc(FLinearColor::Blue, FVector(0, 0, 1));
+		bAnyRotate = true;
+	}
+
+
+	// add a non-interactive view-aligned circle element, so the axes look like a sphere.
+	if (bAnyRotate)
+	{
+		UGizmoCircleComponent* SphereEdge = NewObject<UGizmoCircleComponent>(NewActor);
+		NewActor->AddInstanceComponent(SphereEdge);
+		SphereEdge->AttachToComponent(NewActor->GetRootComponent(), FAttachmentTransformRules::KeepRelativeTransform);
+		SphereEdge->SetGizmoViewContext(GizmoViewContext);
+		SphereEdge->Color = FLinearColor::Gray;
+		SphereEdge->Thickness = 1.0f;
+		SphereEdge->Radius = 120.0f;
+		SphereEdge->bViewAligned = true;
+		SphereEdge->RegisterComponent();
+	}
+
+
+
+	if ((Elements & ETransformGizmoSubElements::ScaleUniform) != ETransformGizmoSubElements::None)
+	{
+		float BoxSize = 14.0f;
+		UGizmoBoxComponent* ScaleComponent = AddDefaultBoxComponent(World, NewActor, GizmoViewContext, FLinearColor::Black,
+			FVector(BoxSize/2, BoxSize/2, BoxSize/2), FVector(BoxSize, BoxSize, BoxSize));
+		NewActor->UniformScale = ScaleComponent;
+	}
+
+
+
+	auto MakeAxisScaleFunc = [&](const FLinearColor& Color, const FVector& Axis0, const FVector& Axis1)
+	{
+		UGizmoRectangleComponent* ScaleComponent = AddDefaultRectangleComponent(World, NewActor, GizmoViewContext, Color, Axis0, Axis1);
+		ScaleComponent->OffsetX = 140.0f; ScaleComponent->OffsetY = -10.0f;
+		ScaleComponent->LengthX = 7.0f; ScaleComponent->LengthY = 20.0f;
+		ScaleComponent->Thickness = GizmoLineThickness;
+		ScaleComponent->bOrientYAccordingToCamera = true;
+		ScaleComponent->NotifyExternalPropertyUpdates();
+		ScaleComponent->SegmentFlags = 0x1 | 0x2 | 0x4; // | 0x8;
+		return ScaleComponent;
+	};
+	if ((Elements & ETransformGizmoSubElements::ScaleAxisX) != ETransformGizmoSubElements::None)
+	{
+		NewActor->AxisScaleX = MakeAxisScaleFunc(FLinearColor::Red, FVector(1, 0, 0), FVector(0, 0, 1));
+	}
+	if ((Elements & ETransformGizmoSubElements::ScaleAxisY) != ETransformGizmoSubElements::None)
+	{
+		NewActor->AxisScaleY = MakeAxisScaleFunc(FLinearColor::Green, FVector(0, 1, 0), FVector(0, 0, 1));
+	}
+	if ((Elements & ETransformGizmoSubElements::ScaleAxisZ) != ETransformGizmoSubElements::None)
+	{
+		NewActor->AxisScaleZ = MakeAxisScaleFunc(FLinearColor::Blue, FVector(0, 0, 1), FVector(1, 0, 0));
+	}
+
+
+	auto MakePlaneScaleFunc = [&](const FLinearColor& Color, const FVector& Axis0, const FVector& Axis1)
+	{
+		UGizmoRectangleComponent* ScaleComponent = AddDefaultRectangleComponent(World, NewActor, GizmoViewContext, Color, Axis0, Axis1);
+		ScaleComponent->OffsetX = ScaleComponent->OffsetY = 120.0f;
+		ScaleComponent->LengthX = ScaleComponent->LengthY = 20.0f;
+		ScaleComponent->Thickness = GizmoLineThickness;
+		ScaleComponent->NotifyExternalPropertyUpdates();
+		ScaleComponent->SegmentFlags = 0x2 | 0x4;
+		return ScaleComponent;
+	};
+	if ((Elements & ETransformGizmoSubElements::ScalePlaneYZ) != ETransformGizmoSubElements::None)
+	{
+		NewActor->PlaneScaleYZ = MakePlaneScaleFunc(FLinearColor::Red, FVector(0, 1, 0), FVector(0, 0, 1));
+	}
+	if ((Elements & ETransformGizmoSubElements::ScalePlaneXZ) != ETransformGizmoSubElements::None)
+	{
+		NewActor->PlaneScaleXZ = MakePlaneScaleFunc(FLinearColor::Green, FVector(1, 0, 0), FVector(0, 0, 1));
+	}
+	if ((Elements & ETransformGizmoSubElements::ScalePlaneXY) != ETransformGizmoSubElements::None)
+	{
+		NewActor->PlaneScaleXY = MakePlaneScaleFunc(FLinearColor::Blue, FVector(1, 0, 0), FVector(0, 1, 0));
+	}
+
+
+	return NewActor;
+}
+
+
+
+
+ACombinedTransformGizmoActor* FCombinedTransformGizmoActorFactory::CreateNewGizmoActor(UWorld* World) const
+{
+	return ACombinedTransformGizmoActor::ConstructCustom3AxisGizmo(World, GizmoViewContext, EnableElements);
+}
+
+
+
+UInteractiveGizmo* UCombinedTransformGizmoBuilder::BuildGizmo(const FToolBuilderState& SceneState) const
+{
+	UCombinedTransformGizmo* NewGizmo = NewObject<UCombinedTransformGizmo>(SceneState.GizmoManager);
+	NewGizmo->SetWorld(SceneState.World);
+
+	UGizmoViewContext* GizmoViewContext = SceneState.ToolManager->GetContextObjectStore()->FindContext<UGizmoViewContext>();
+	check(GizmoViewContext && GizmoViewContext->IsValidLowLevel());
+
+	// use default gizmo actor if client has not given us a new builder
+	NewGizmo->SetGizmoActorBuilder(GizmoActorBuilder ? GizmoActorBuilder : MakeShared<FCombinedTransformGizmoActorFactory>(GizmoViewContext));
+
+	NewGizmo->SetSubGizmoBuilderIdentifiers(AxisPositionBuilderIdentifier, PlanePositionBuilderIdentifier, AxisAngleBuilderIdentifier);
+
+	// override default hover function if proposed
+	if (UpdateHoverFunction)
+	{
+		NewGizmo->SetUpdateHoverFunction(UpdateHoverFunction);
+	}
+
+	if (UpdateCoordSystemFunction)
+	{
+		NewGizmo->SetUpdateCoordSystemFunction(UpdateCoordSystemFunction);
+	}
+
+	return NewGizmo;
+}
+
+
+
+void UCombinedTransformGizmo::SetWorld(UWorld* WorldIn)
+{
+	this->World = WorldIn;
+}
+
+void UCombinedTransformGizmo::SetGizmoActorBuilder(TSharedPtr<FCombinedTransformGizmoActorFactory> Builder)
+{
+	GizmoActorBuilder = Builder;
+}
+
+void UCombinedTransformGizmo::SetSubGizmoBuilderIdentifiers(FString AxisPositionBuilderIdentifierIn, FString PlanePositionBuilderIdentifierIn, FString AxisAngleBuilderIdentifierIn)
+{
+	AxisPositionBuilderIdentifier = AxisPositionBuilderIdentifierIn;
+	PlanePositionBuilderIdentifier = PlanePositionBuilderIdentifierIn;
+	AxisAngleBuilderIdentifier = AxisAngleBuilderIdentifierIn;
+}
+
+void UCombinedTransformGizmo::SetUpdateHoverFunction(TFunction<void(UPrimitiveComponent*, bool)> HoverFunction)
+{
+	UpdateHoverFunction = HoverFunction;
+}
+
+void UCombinedTransformGizmo::SetUpdateCoordSystemFunction(TFunction<void(UPrimitiveComponent*, EToolContextCoordinateSystem)> CoordSysFunction)
+{
+	UpdateCoordSystemFunction = CoordSysFunction;
+}
+
+void UCombinedTransformGizmo::SetWorldAlignmentFunctions(
+	TUniqueFunction<bool()>&& ShouldAlignTranslationIn,
+	TUniqueFunction<bool(const FRay&, FVector&)>&& TranslationAlignmentRayCasterIn)
+{
+	// Save these so that later changes of gizmo target keep the settings.
+	ShouldAlignDestination = MoveTemp(ShouldAlignTranslationIn);
+	DestinationAlignmentRayCaster = MoveTemp(TranslationAlignmentRayCasterIn);
+
+	// We allow this function to be called after Setup(), so modify any existing translation/rotation sub gizmos.
+	// Unfortunately we keep all the sub gizmos in one list, and the scaling gizmos are differentiated from the
+	// translation ones mainly in the components they use. So this ends up being a slightly messy set of checks,
+	// but it didn't seem worth keeping a segregated list for something that will only happen once.
+	for (UInteractiveGizmo* SubGizmo : this->ActiveGizmos)
+	{
+		if (UAxisPositionGizmo* CastGizmo = Cast<UAxisPositionGizmo>(SubGizmo))
+		{
+			if (UGizmoComponentHitTarget* CastHitTarget = Cast<UGizmoComponentHitTarget>(CastGizmo->HitTarget.GetObject()))
+			{
+				if (CastHitTarget->Component == GizmoActor->TranslateX
+					|| CastHitTarget->Component == GizmoActor->TranslateY
+					|| CastHitTarget->Component == GizmoActor->TranslateZ)
+				{
+					CastGizmo->ShouldUseCustomDestinationFunc = [this]() { return ShouldAlignDestination(); };
+					CastGizmo->CustomDestinationFunc = 
+						[this](const UAxisPositionGizmo::FCustomDestinationParams& Params, FVector& OutputPoint) { 
+						return DestinationAlignmentRayCaster(*Params.WorldRay, OutputPoint); 
+					};
+				}
+			}
+		}
+		if (UPlanePositionGizmo* CastGizmo = Cast<UPlanePositionGizmo>(SubGizmo))
+		{
+			if (UGizmoComponentHitTarget* CastHitTarget = Cast<UGizmoComponentHitTarget>(CastGizmo->HitTarget.GetObject()))
+			{
+				if (CastHitTarget->Component == GizmoActor->TranslateXY
+					|| CastHitTarget->Component == GizmoActor->TranslateXZ
+					|| CastHitTarget->Component == GizmoActor->TranslateYZ)
+				{
+					CastGizmo->ShouldUseCustomDestinationFunc = [this]() { return ShouldAlignDestination(); };
+					CastGizmo->CustomDestinationFunc =
+						[this](const UPlanePositionGizmo::FCustomDestinationParams& Params, FVector& OutputPoint) {
+						return DestinationAlignmentRayCaster(*Params.WorldRay, OutputPoint);
+					};
+				}
+			}
+		}
+		if (UAxisAngleGizmo* CastGizmo = Cast<UAxisAngleGizmo>(SubGizmo))
+		{
+			CastGizmo->ShouldUseCustomDestinationFunc = [this]() { return ShouldAlignDestination(); };
+			CastGizmo->CustomDestinationFunc =
+				[this](const UAxisAngleGizmo::FCustomDestinationParams& Params, FVector& OutputPoint) {
+				return DestinationAlignmentRayCaster(*Params.WorldRay, OutputPoint);
+			};
+		}
+	}
+}
+
+void UCombinedTransformGizmo::SetDisallowNegativeScaling(bool bDisallow)
+{
+	if (bDisallowNegativeScaling != bDisallow)
+	{
+		bDisallowNegativeScaling = bDisallow;
+		for (UInteractiveGizmo* SubGizmo : this->ActiveGizmos)
+		{
+			if (UAxisPositionGizmo* CastGizmo = Cast<UAxisPositionGizmo>(SubGizmo))
+			{
+				if (UGizmoAxisScaleParameterSource* ParamSource = Cast<UGizmoAxisScaleParameterSource>(CastGizmo->ParameterSource.GetObject()))
+				{
+					ParamSource->bClampToZero = bDisallow;
+				}
+			}
+			if (UPlanePositionGizmo* CastGizmo = Cast<UPlanePositionGizmo>(SubGizmo))
+			{
+				if (UGizmoPlaneScaleParameterSource* ParamSource = Cast<UGizmoPlaneScaleParameterSource>(CastGizmo->ParameterSource.GetObject()))
+				{
+					ParamSource->bClampToZero = bDisallow;
+				}
+			}
+		}
+	}
+}
+
+void UCombinedTransformGizmo::SetIsNonUniformScaleAllowedFunction(TUniqueFunction<bool()>&& IsNonUniformScaleAllowedIn)
+{
+	IsNonUniformScaleAllowed = MoveTemp(IsNonUniformScaleAllowedIn);
+}
+
+
+void UCombinedTransformGizmo::Setup()
+{
+	UInteractiveGizmo::Setup();
+
+	UpdateHoverFunction = [](UPrimitiveComponent* Component, bool bHovering)
+	{
+		if (Cast<UGizmoBaseComponent>(Component) != nullptr)
+		{
+			Cast<UGizmoBaseComponent>(Component)->UpdateHoverState(bHovering);
+		}
+	};
+
+	UpdateCoordSystemFunction = [](UPrimitiveComponent* Component, EToolContextCoordinateSystem CoordSystem)
+	{
+		if (Cast<UGizmoBaseComponent>(Component) != nullptr)
+		{
+			Cast<UGizmoBaseComponent>(Component)->UpdateWorldLocalState(CoordSystem == EToolContextCoordinateSystem::World);
+		}
+	};
+
+	GizmoActor = GizmoActorBuilder->CreateNewGizmoActor(World);
+}
+
+
+
+void UCombinedTransformGizmo::Shutdown()
+{
+	ClearActiveTarget();
+
+	if (GizmoActor)
+	{
+		GizmoActor->Destroy();
+		GizmoActor = nullptr;
+	}
+}
+
+
+
+void UCombinedTransformGizmo::UpdateCameraAxisSource()
+{
+	FViewCameraState CameraState;
+	GetGizmoManager()->GetContextQueriesAPI()->GetCurrentViewState(CameraState);
+	if (CameraAxisSource != nullptr && GizmoActor != nullptr)
+	{
+		CameraAxisSource->Origin = GizmoActor->GetTransform().GetLocation();
+		CameraAxisSource->Direction = -CameraState.Forward();
+		CameraAxisSource->TangentX = CameraState.Right();
+		CameraAxisSource->TangentY = CameraState.Up();
+	}
+}
+
+
+void UCombinedTransformGizmo::Tick(float DeltaTime)
+{	
+	if (bUseContextCoordinateSystem)
+	{
+		CurrentCoordinateSystem = GetGizmoManager()->GetContextQueriesAPI()->GetCurrentCoordinateSystem();
+	}
+	
+	check(CurrentCoordinateSystem == EToolContextCoordinateSystem::World || CurrentCoordinateSystem == EToolContextCoordinateSystem::Local)
+	bool bUseLocalAxes = (CurrentCoordinateSystem == EToolContextCoordinateSystem::Local);
+
+	if (AxisXSource != nullptr && AxisYSource != nullptr && AxisZSource != nullptr)
+	{
+		AxisXSource->bLocalAxes = bUseLocalAxes;
+		AxisYSource->bLocalAxes = bUseLocalAxes;
+		AxisZSource->bLocalAxes = bUseLocalAxes;
+	}
+	if (UpdateCoordSystemFunction)
+	{
+		for (UPrimitiveComponent* Component : ActiveComponents)
+		{
+			UpdateCoordSystemFunction(Component, CurrentCoordinateSystem);
+		}
+	}
+
+	bool bShouldShowNonUniformScale = IsNonUniformScaleAllowed();
+	for (UPrimitiveComponent* Component : NonuniformScaleComponents)
+	{
+		Component->SetVisibility(bShouldShowNonUniformScale);
+	}
+
+	UpdateCameraAxisSource();
+}
+
+
+
+void UCombinedTransformGizmo::SetActiveTarget(UTransformProxy* Target, IToolContextTransactionProvider* TransactionProvider)
+{
+	if (ActiveTarget != nullptr)
+	{
+		ClearActiveTarget();
+	}
+
+	ActiveTarget = Target;
+
+	// move gizmo to target location
+	USceneComponent* GizmoComponent = GizmoActor->GetRootComponent();
+
+	FTransform TargetTransform = Target->GetTransform();
+	FTransform GizmoTransform = TargetTransform;
+	GizmoTransform.SetScale3D(FVector(1, 1, 1));
+	GizmoComponent->SetWorldTransform(GizmoTransform);
+
+	UGizmoScaledAndUnscaledTransformSources* TransformSource = UGizmoScaledAndUnscaledTransformSources::Construct(
+		UGizmoTransformProxyTransformSource::Construct(ActiveTarget, this), 
+		GizmoComponent, this);
+
+	// This state target emits an explicit FChange that moves the GizmoActor root component during undo/redo.
+	// It also opens/closes the Transaction that saves/restores the target object locations.
+	if (TransactionProvider == nullptr)
+	{
+		TransactionProvider = GetGizmoManager();
+	}
+	StateTarget = UGizmoTransformChangeStateTarget::Construct(GizmoComponent,
+		LOCTEXT("UCombinedTransformGizmoTransaction", "Transform"), TransactionProvider, this);
+	StateTarget->DependentChangeSources.Add(MakeUnique<FTransformProxyChangeSource>(Target));
+
+	CameraAxisSource = NewObject<UGizmoConstantFrameAxisSource>(this);
+
+	// root component provides local X/Y/Z axis, identified by AxisIndex
+	AxisXSource = UGizmoComponentAxisSource::Construct(GizmoComponent, 0, true, this);
+	AxisYSource = UGizmoComponentAxisSource::Construct(GizmoComponent, 1, true, this);
+	AxisZSource = UGizmoComponentAxisSource::Construct(GizmoComponent, 2, true, this);
+
+	// todo should we hold onto these?
+	if (GizmoActor->TranslateX != nullptr)
+	{
+		AddAxisTranslationGizmo(GizmoActor->TranslateX, GizmoComponent, AxisXSource, TransformSource, StateTarget);
+		ActiveComponents.Add(GizmoActor->TranslateX);
+	}
+	if (GizmoActor->TranslateY != nullptr)
+	{
+		AddAxisTranslationGizmo(GizmoActor->TranslateY, GizmoComponent, AxisYSource, TransformSource, StateTarget);
+		ActiveComponents.Add(GizmoActor->TranslateY);
+	}
+	if (GizmoActor->TranslateZ != nullptr)
+	{
+		AddAxisTranslationGizmo(GizmoActor->TranslateZ, GizmoComponent, AxisZSource, TransformSource, StateTarget);
+		ActiveComponents.Add(GizmoActor->TranslateZ);
+	}
+
+
+	if (GizmoActor->TranslateYZ != nullptr)
+	{
+		AddPlaneTranslationGizmo(GizmoActor->TranslateYZ, GizmoComponent, AxisXSource, TransformSource, StateTarget);
+		ActiveComponents.Add(GizmoActor->TranslateYZ);
+	}
+	if (GizmoActor->TranslateXZ != nullptr)
+	{
+		AddPlaneTranslationGizmo(GizmoActor->TranslateXZ, GizmoComponent, AxisYSource, TransformSource, StateTarget);
+		ActiveComponents.Add(GizmoActor->TranslateXZ);
+	}
+	if (GizmoActor->TranslateXY != nullptr)
+	{
+		AddPlaneTranslationGizmo(GizmoActor->TranslateXY, GizmoComponent, AxisZSource, TransformSource, StateTarget);
+		ActiveComponents.Add(GizmoActor->TranslateXY);
+	}
+
+	if (GizmoActor->RotateX != nullptr)
+	{
+		AddAxisRotationGizmo(GizmoActor->RotateX, GizmoComponent, AxisXSource, TransformSource, StateTarget);
+		ActiveComponents.Add(GizmoActor->RotateX);
+	}
+	if (GizmoActor->RotateY != nullptr)
+	{
+		AddAxisRotationGizmo(GizmoActor->RotateY, GizmoComponent, AxisYSource, TransformSource, StateTarget);
+		ActiveComponents.Add(GizmoActor->RotateY);
+	}
+	if (GizmoActor->RotateZ != nullptr)
+	{
+		AddAxisRotationGizmo(GizmoActor->RotateZ, GizmoComponent, AxisZSource, TransformSource, StateTarget);
+		ActiveComponents.Add(GizmoActor->RotateZ);
+	}
+
+
+	// only need these if scaling enabled. Essentially these are just the unit axes, regardless
+	// of what 3D axis is in use, we will tell the ParameterSource-to-3D-Scale mapper to
+	// use the coordinate axes
+	UnitAxisXSource = UGizmoComponentAxisSource::Construct(GizmoComponent, 0, false, this);
+	UnitAxisYSource = UGizmoComponentAxisSource::Construct(GizmoComponent, 1, false, this);
+	UnitAxisZSource = UGizmoComponentAxisSource::Construct(GizmoComponent, 2, false, this);
+
+	if (GizmoActor->UniformScale != nullptr)
+	{
+		AddUniformScaleGizmo(GizmoActor->UniformScale, GizmoComponent, CameraAxisSource, CameraAxisSource, TransformSource, StateTarget);
+		ActiveComponents.Add(GizmoActor->UniformScale);
+	}
+
+	if (GizmoActor->AxisScaleX != nullptr)
+	{
+		AddAxisScaleGizmo(GizmoActor->AxisScaleX, GizmoComponent, AxisXSource, UnitAxisXSource, TransformSource, StateTarget);
+		ActiveComponents.Add(GizmoActor->AxisScaleX);
+		NonuniformScaleComponents.Add(GizmoActor->AxisScaleX);
+	}
+	if (GizmoActor->AxisScaleY != nullptr)
+	{
+		AddAxisScaleGizmo(GizmoActor->AxisScaleY, GizmoComponent, AxisYSource, UnitAxisYSource, TransformSource, StateTarget);
+		ActiveComponents.Add(GizmoActor->AxisScaleY);
+		NonuniformScaleComponents.Add(GizmoActor->AxisScaleY);
+	}
+	if (GizmoActor->AxisScaleZ != nullptr)
+	{
+		AddAxisScaleGizmo(GizmoActor->AxisScaleZ, GizmoComponent, AxisZSource, UnitAxisZSource, TransformSource, StateTarget);
+		ActiveComponents.Add(GizmoActor->AxisScaleZ);
+		NonuniformScaleComponents.Add(GizmoActor->AxisScaleZ);
+	}
+
+	if (GizmoActor->PlaneScaleYZ != nullptr)
+	{
+		AddPlaneScaleGizmo(GizmoActor->PlaneScaleYZ, GizmoComponent, AxisXSource, UnitAxisXSource, TransformSource, StateTarget);
+		ActiveComponents.Add(GizmoActor->PlaneScaleYZ);
+		NonuniformScaleComponents.Add(GizmoActor->PlaneScaleYZ);
+	}
+	if (GizmoActor->PlaneScaleXZ != nullptr)
+	{
+		UPlanePositionGizmo* Gizmo = (UPlanePositionGizmo *)AddPlaneScaleGizmo(GizmoActor->PlaneScaleXZ, GizmoComponent, AxisYSource, UnitAxisYSource, TransformSource, StateTarget);
+		ActiveComponents.Add(GizmoActor->PlaneScaleXZ);
+		NonuniformScaleComponents.Add(GizmoActor->PlaneScaleXZ);
+	}
+	if (GizmoActor->PlaneScaleXY != nullptr)
+	{
+		AddPlaneScaleGizmo(GizmoActor->PlaneScaleXY, GizmoComponent, AxisZSource, UnitAxisZSource, TransformSource, StateTarget);
+		ActiveComponents.Add(GizmoActor->PlaneScaleXY);
+		NonuniformScaleComponents.Add(GizmoActor->PlaneScaleXY);
+	}
+}
+
+FTransform UCombinedTransformGizmo::GetGizmoTransform() const
+{
+	USceneComponent* GizmoComponent = GizmoActor->GetRootComponent();
+	return GizmoComponent->GetComponentTransform();
+}
+
+void UCombinedTransformGizmo::ReinitializeGizmoTransform(const FTransform& NewTransform, bool bKeepGizmoUnscaled)
+{
+	// To update the gizmo location without triggering any callbacks, we temporarily
+	// store a copy of the callback list, detach them, reposition, and then reattach
+	// the callbacks.
+	USceneComponent* GizmoComponent = GizmoActor->GetRootComponent();
+	auto temp = GizmoComponent->TransformUpdated;
+	GizmoComponent->TransformUpdated.Clear();
+	FTransform GizmoTransform = NewTransform;
+	if (bKeepGizmoUnscaled)
+	{
+		GizmoTransform.SetScale3D(FVector(1, 1, 1));
+	}
+	GizmoComponent->SetWorldTransform(GizmoTransform);
+	GizmoComponent->TransformUpdated = temp;
+
+	// The underlying proxy has an existing way to reinitialize its transform without callbacks.
+	bool bSavedSetPivotMode = ActiveTarget->bSetPivotMode;
+	ActiveTarget->bSetPivotMode = true;
+	ActiveTarget->SetTransform(NewTransform);
+	ActiveTarget->bSetPivotMode = bSavedSetPivotMode;
+}
+
+void UCombinedTransformGizmo::SetNewGizmoTransform(const FTransform& NewTransform, bool bKeepGizmoUnscaled)
+{
+	check(ActiveTarget != nullptr);
+
+	StateTarget->BeginUpdate();
+
+	USceneComponent* GizmoComponent = GizmoActor->GetRootComponent();
+	FTransform GizmoTransform = NewTransform;
+	if (bKeepGizmoUnscaled)
+	{
+		GizmoTransform.SetScale3D(FVector(1, 1, 1));
+	}
+	GizmoComponent->SetWorldTransform(GizmoTransform);
+	ActiveTarget->SetTransform(NewTransform);
+
+	StateTarget->EndUpdate();
+}
+
+
+void UCombinedTransformGizmo::SetNewChildScale(const FVector& NewChildScale)
+{
+	FTransform NewTransform = ActiveTarget->GetTransform();
+	NewTransform.SetScale3D(NewChildScale);
+
+	bool bSavedSetPivotMode = ActiveTarget->bSetPivotMode;
+	ActiveTarget->bSetPivotMode = true;
+	ActiveTarget->SetTransform(NewTransform);
+	ActiveTarget->bSetPivotMode = bSavedSetPivotMode;
+}
+
+
+void UCombinedTransformGizmo::SetVisibility(bool bVisible)
+{
+	GizmoActor->SetActorHiddenInGame(bVisible == false);
+#if WITH_EDITOR
+	GizmoActor->SetIsTemporarilyHiddenInEditor(bVisible == false);
+#endif
+}
+
+
+UInteractiveGizmo* UCombinedTransformGizmo::AddAxisTranslationGizmo(
+	UPrimitiveComponent* AxisComponent, USceneComponent* RootComponent,
+	IGizmoAxisSource* AxisSource,
+	IGizmoTransformSource* TransformSource,
+	IGizmoStateTarget* StateTargetIn)
+{
+	// create axis-position gizmo, axis-position parameter will drive translation
+	UAxisPositionGizmo* TranslateGizmo = Cast<UAxisPositionGizmo>(GetGizmoManager()->CreateGizmo(AxisPositionBuilderIdentifier));
+	check(TranslateGizmo);
+
+	// axis source provides the translation axis
+	TranslateGizmo->AxisSource = Cast<UObject>(AxisSource);
+
+	// parameter source maps axis-parameter-change to translation of TransformSource's transform
+	UGizmoAxisTranslationParameterSource* ParamSource = UGizmoAxisTranslationParameterSource::Construct(AxisSource, TransformSource, this);
+	ParamSource->PositionConstraintFunction = [this](const FVector& Pos, FVector& Snapped) { return PositionSnapFunction(Pos, Snapped); };
+	TranslateGizmo->ParameterSource = ParamSource;
+
+	// sub-component provides hit target
+	UGizmoComponentHitTarget* HitTarget = UGizmoComponentHitTarget::Construct(AxisComponent, this);
+	if (this->UpdateHoverFunction)
+	{
+		HitTarget->UpdateHoverFunction = [AxisComponent, this](bool bHovering) { this->UpdateHoverFunction(AxisComponent, bHovering); };
+	}
+	TranslateGizmo->HitTarget = HitTarget;
+
+	TranslateGizmo->StateTarget = Cast<UObject>(StateTargetIn);
+
+	TranslateGizmo->ShouldUseCustomDestinationFunc = [this]() { return ShouldAlignDestination(); };
+	TranslateGizmo->CustomDestinationFunc =
+		[this](const UAxisPositionGizmo::FCustomDestinationParams& Params, FVector& OutputPoint) {
+		return DestinationAlignmentRayCaster(*Params.WorldRay, OutputPoint);
+	};
+
+	ActiveGizmos.Add(TranslateGizmo);
+	return TranslateGizmo;
+}
+
+
+
+UInteractiveGizmo* UCombinedTransformGizmo::AddPlaneTranslationGizmo(
+	UPrimitiveComponent* AxisComponent, USceneComponent* RootComponent,
+	IGizmoAxisSource* AxisSource,
+	IGizmoTransformSource* TransformSource,
+	IGizmoStateTarget* StateTargetIn)
+{
+	// create axis-position gizmo, axis-position parameter will drive translation
+	UPlanePositionGizmo* TranslateGizmo = Cast<UPlanePositionGizmo>(GetGizmoManager()->CreateGizmo(PlanePositionBuilderIdentifier));
+	check(TranslateGizmo);
+
+	// axis source provides the translation axis
+	TranslateGizmo->AxisSource = Cast<UObject>(AxisSource);
+
+	// parameter source maps axis-parameter-change to translation of TransformSource's transform
+	UGizmoPlaneTranslationParameterSource* ParamSource = UGizmoPlaneTranslationParameterSource::Construct(AxisSource, TransformSource, this);
+	ParamSource->PositionConstraintFunction = [this](const FVector& Pos, FVector& Snapped) { return PositionSnapFunction(Pos, Snapped); };
+	TranslateGizmo->ParameterSource = ParamSource;
+
+	// sub-component provides hit target
+	UGizmoComponentHitTarget* HitTarget = UGizmoComponentHitTarget::Construct(AxisComponent, this);
+	if (this->UpdateHoverFunction)
+	{
+		HitTarget->UpdateHoverFunction = [AxisComponent, this](bool bHovering) { this->UpdateHoverFunction(AxisComponent, bHovering); };
+	}
+	TranslateGizmo->HitTarget = HitTarget;
+
+	TranslateGizmo->StateTarget = Cast<UObject>(StateTargetIn);
+
+	TranslateGizmo->ShouldUseCustomDestinationFunc = [this]() { return ShouldAlignDestination(); };
+	TranslateGizmo->CustomDestinationFunc =
+		[this](const UPlanePositionGizmo::FCustomDestinationParams& Params, FVector& OutputPoint) {
+		return DestinationAlignmentRayCaster(*Params.WorldRay, OutputPoint);
+	};
+
+	ActiveGizmos.Add(TranslateGizmo);
+	return TranslateGizmo;
+}
+
+
+
+
+
+UInteractiveGizmo* UCombinedTransformGizmo::AddAxisRotationGizmo(
+	UPrimitiveComponent* AxisComponent, USceneComponent* RootComponent,
+	IGizmoAxisSource* AxisSource,
+	IGizmoTransformSource* TransformSource,
+	IGizmoStateTarget* StateTargetIn)
+{
+	// create axis-angle gizmo, angle will drive axis-rotation
+	UAxisAngleGizmo* RotateGizmo = Cast<UAxisAngleGizmo>(GetGizmoManager()->CreateGizmo(AxisAngleBuilderIdentifier));
+	check(RotateGizmo);
+
+	// axis source provides the rotation axis
+	RotateGizmo->AxisSource = Cast<UObject>(AxisSource);
+
+	// parameter source maps angle-parameter-change to rotation of TransformSource's transform
+	UGizmoAxisRotationParameterSource* AngleSource = UGizmoAxisRotationParameterSource::Construct(AxisSource, TransformSource, this);
+	AngleSource->RotationConstraintFunction = [this](const FQuat& DeltaRotation){ return RotationSnapFunction(DeltaRotation); };
+	RotateGizmo->AngleSource = AngleSource;
+
+	// sub-component provides hit target
+	UGizmoComponentHitTarget* HitTarget = UGizmoComponentHitTarget::Construct(AxisComponent, this);
+	if (this->UpdateHoverFunction)
+	{
+		HitTarget->UpdateHoverFunction = [AxisComponent, this](bool bHovering) { this->UpdateHoverFunction(AxisComponent, bHovering); };
+	}
+	RotateGizmo->HitTarget = HitTarget;
+
+	RotateGizmo->StateTarget = Cast<UObject>(StateTargetIn);
+
+	RotateGizmo->ShouldUseCustomDestinationFunc = [this]() { return ShouldAlignDestination(); };
+	RotateGizmo->CustomDestinationFunc =
+		[this](const UAxisAngleGizmo::FCustomDestinationParams& Params, FVector& OutputPoint) {
+		return DestinationAlignmentRayCaster(*Params.WorldRay, OutputPoint);
+	};
+
+	ActiveGizmos.Add(RotateGizmo);
+
+	return RotateGizmo;
+}
+
+
+
+UInteractiveGizmo* UCombinedTransformGizmo::AddAxisScaleGizmo(
+	UPrimitiveComponent* AxisComponent, USceneComponent* RootComponent,
+	IGizmoAxisSource* GizmoAxisSource, IGizmoAxisSource* ParameterAxisSource,
+	IGizmoTransformSource* TransformSource,
+	IGizmoStateTarget* StateTargetIn)
+{
+	// create axis-position gizmo, axis-position parameter will drive scale
+	UAxisPositionGizmo* ScaleGizmo = Cast<UAxisPositionGizmo>(GetGizmoManager()->CreateGizmo(AxisPositionBuilderIdentifier));
+	ScaleGizmo->bEnableSignedAxis = true;
+	check(ScaleGizmo);
+
+	// axis source provides the translation axis
+	ScaleGizmo->AxisSource = Cast<UObject>(GizmoAxisSource);
+
+	// parameter source maps axis-parameter-change to translation of TransformSource's transform
+	UGizmoAxisScaleParameterSource* ParamSource = UGizmoAxisScaleParameterSource::Construct(ParameterAxisSource, TransformSource, this);
+	ParamSource->bClampToZero = bDisallowNegativeScaling;
+	ScaleGizmo->ParameterSource = ParamSource;
+
+	// sub-component provides hit target
+	UGizmoComponentHitTarget* HitTarget = UGizmoComponentHitTarget::Construct(AxisComponent, this);
+	if (this->UpdateHoverFunction)
+	{
+		HitTarget->UpdateHoverFunction = [AxisComponent, this](bool bHovering) { this->UpdateHoverFunction(AxisComponent, bHovering); };
+	}
+	ScaleGizmo->HitTarget = HitTarget;
+
+	ScaleGizmo->StateTarget = Cast<UObject>(StateTargetIn);
+
+	ActiveGizmos.Add(ScaleGizmo);
+	return ScaleGizmo;
+}
+
+
+
+UInteractiveGizmo* UCombinedTransformGizmo::AddPlaneScaleGizmo(
+	UPrimitiveComponent* AxisComponent, USceneComponent* RootComponent,
+	IGizmoAxisSource* GizmoAxisSource, IGizmoAxisSource* ParameterAxisSource,
+	IGizmoTransformSource* TransformSource,
+	IGizmoStateTarget* StateTargetIn)
+{
+	// create axis-position gizmo, axis-position parameter will drive scale
+	UPlanePositionGizmo* ScaleGizmo = Cast<UPlanePositionGizmo>(GetGizmoManager()->CreateGizmo(PlanePositionBuilderIdentifier));
+	ScaleGizmo->bEnableSignedAxis = true;
+	check(ScaleGizmo);
+
+	// axis source provides the translation axis
+	ScaleGizmo->AxisSource = Cast<UObject>(GizmoAxisSource);
+
+	// parameter source maps axis-parameter-change to translation of TransformSource's transform
+	UGizmoPlaneScaleParameterSource* ParamSource = UGizmoPlaneScaleParameterSource::Construct(ParameterAxisSource, TransformSource, this);
+	ParamSource->bClampToZero = bDisallowNegativeScaling;
+	ParamSource->bUseEqualScaling = true;
+	ScaleGizmo->ParameterSource = ParamSource;
+
+	// sub-component provides hit target
+	UGizmoComponentHitTarget* HitTarget = UGizmoComponentHitTarget::Construct(AxisComponent, this);
+	if (this->UpdateHoverFunction)
+	{
+		HitTarget->UpdateHoverFunction = [AxisComponent, this](bool bHovering) { this->UpdateHoverFunction(AxisComponent, bHovering); };
+	}
+	ScaleGizmo->HitTarget = HitTarget;
+
+	ScaleGizmo->StateTarget = Cast<UObject>(StateTargetIn);
+
+	ActiveGizmos.Add(ScaleGizmo);
+	return ScaleGizmo;
+}
+
+
+
+
+
+UInteractiveGizmo* UCombinedTransformGizmo::AddUniformScaleGizmo(
+	UPrimitiveComponent* ScaleComponent, USceneComponent* RootComponent,
+	IGizmoAxisSource* GizmoAxisSource, IGizmoAxisSource* ParameterAxisSource,
+	IGizmoTransformSource* TransformSource,
+	IGizmoStateTarget* StateTargetIn)
+{
+	// create plane-position gizmo, plane-position parameter will drive scale
+	UPlanePositionGizmo* ScaleGizmo = Cast<UPlanePositionGizmo>(GetGizmoManager()->CreateGizmo(PlanePositionBuilderIdentifier));
+	check(ScaleGizmo);
+
+	// axis source provides the translation plane
+	ScaleGizmo->AxisSource = Cast<UObject>(GizmoAxisSource);
+
+	// parameter source maps axis-parameter-change to translation of TransformSource's transform
+	UGizmoUniformScaleParameterSource* ParamSource = UGizmoUniformScaleParameterSource::Construct(ParameterAxisSource, TransformSource, this);
+	//ParamSource->PositionConstraintFunction = [this](const FVector& Pos, FVector& Snapped) { return PositionSnapFunction(Pos, Snapped); };
+	ScaleGizmo->ParameterSource = ParamSource;
+
+	// sub-component provides hit target
+	UGizmoComponentHitTarget* HitTarget = UGizmoComponentHitTarget::Construct(ScaleComponent, this);
+	if (this->UpdateHoverFunction)
+	{
+		HitTarget->UpdateHoverFunction = [ScaleComponent, this](bool bHovering) { this->UpdateHoverFunction(ScaleComponent, bHovering); };
+	}
+	ScaleGizmo->HitTarget = HitTarget;
+
+	ScaleGizmo->StateTarget = Cast<UObject>(StateTargetIn);
+
+	ActiveGizmos.Add(ScaleGizmo);
+	return ScaleGizmo;
+}
+
+
+
+void UCombinedTransformGizmo::ClearActiveTarget()
+{
+	for (UInteractiveGizmo* Gizmo : ActiveGizmos)
+	{
+		GetGizmoManager()->DestroyGizmo(Gizmo);
+	}
+	ActiveGizmos.SetNum(0);
+	ActiveComponents.SetNum(0);
+	NonuniformScaleComponents.SetNum(0);
+
+	CameraAxisSource = nullptr;
+	AxisXSource = nullptr;
+	AxisYSource = nullptr;
+	AxisZSource = nullptr;
+	UnitAxisXSource = nullptr;
+	UnitAxisYSource = nullptr;
+	UnitAxisZSource = nullptr;
+	StateTarget = nullptr;
+
+	ActiveTarget = nullptr;
+}
+
+
+
+
+bool UCombinedTransformGizmo::PositionSnapFunction(const FVector& WorldPosition, FVector& SnappedPositionOut) const
+{
+	SnappedPositionOut = WorldPosition;
+
+	// only snap if we want snapping obvs
+	if (bSnapToWorldGrid == false)
+	{
+		return false;
+	}
+
+	// only snap to world grid when using world axes
+	EToolContextCoordinateSystem CoordSystem = GetGizmoManager()->GetContextQueriesAPI()->GetCurrentCoordinateSystem();
+	if (CoordSystem != EToolContextCoordinateSystem::World)
+	{
+		return false;
+	}
+
+	FSceneSnapQueryRequest Request;
+	Request.RequestType = ESceneSnapQueryType::Position;
+	Request.TargetTypes = ESceneSnapQueryTargetType::Grid;
+	Request.Position = WorldPosition;
+	if ( bGridSizeIsExplicit )
+	{
+		Request.GridSize = ExplicitGridSize;
+	}
+	TArray<FSceneSnapQueryResult> Results;
+	if (GetGizmoManager()->GetContextQueriesAPI()->ExecuteSceneSnapQuery(Request, Results))
+	{
+		SnappedPositionOut = Results[0].Position;
+		return true;
+	};
+
+	return false;
+}
+
+FQuat UCombinedTransformGizmo::RotationSnapFunction(const FQuat& DeltaRotation) const
+{
+	FQuat SnappedDeltaRotation = DeltaRotation;
+
+	// only snap if we want snapping obvs
+	if (bSnapToWorldRotGrid)
+	{
+		FSceneSnapQueryRequest Request;
+		Request.RequestType   = ESceneSnapQueryType::Rotation;
+		Request.TargetTypes   = ESceneSnapQueryTargetType::Grid;
+		Request.DeltaRotation = DeltaRotation;
+		if ( bRotationGridSizeIsExplicit )
+		{
+			Request.RotGridSize = ExplicitRotationGridSize;
+		}
+		TArray<FSceneSnapQueryResult> Results;
+		if (GetGizmoManager()->GetContextQueriesAPI()->ExecuteSceneSnapQuery(Request, Results))
+		{
+			SnappedDeltaRotation = Results[0].DeltaRotation;
+		};
+	}
+	return SnappedDeltaRotation;
+}
+
+#undef LOCTEXT_NAMESPACE
diff --git a/Engine/Source/Runtime/InteractiveToolsFramework/Private/BaseGizmos/TransformGizmo.cpp b/Engine/Source/Runtime/InteractiveToolsFramework/Private/BaseGizmos/TransformGizmo.cpp
deleted file mode 100644
diff --git a/Engine/Source/Runtime/InteractiveToolsFramework/Private/BaseGizmos/TransformGizmoUtil.cpp b/Engine/Source/Runtime/InteractiveToolsFramework/Private/BaseGizmos/TransformGizmoUtil.cpp
--- a/Engine/Source/Runtime/InteractiveToolsFramework/Private/BaseGizmos/TransformGizmoUtil.cpp
+++ b/Engine/Source/Runtime/InteractiveToolsFramework/Private/BaseGizmos/TransformGizmoUtil.cpp
@@ -11,7 +11,7 @@
 #include "BaseGizmos/PlanePositionGizmo.h"
 #include "BaseGizmos/AxisAngleGizmo.h"
 #include "BaseGizmos/GizmoViewContext.h"
-#include "BaseGizmos/TransformGizmo.h"
+#include "BaseGizmos/CombinedTransformGizmo.h"
 #include "BaseGizmos/RepositionableTransformGizmo.h"
 
 const FString UCombinedTransformGizmoContextObject::DefaultAxisPositionBuilderIdentifier = TEXT("Util_StandardXFormAxisTranslationGizmo");
diff --git a/Engine/Source/Runtime/InteractiveToolsFramework/Private/InteractiveGizmoManager.cpp b/Engine/Source/Runtime/InteractiveToolsFramework/Private/InteractiveGizmoManager.cpp
--- a/Engine/Source/Runtime/InteractiveToolsFramework/Private/InteractiveGizmoManager.cpp
+++ b/Engine/Source/Runtime/InteractiveToolsFramework/Private/InteractiveGizmoManager.cpp
@@ -8,7 +8,7 @@
 #include "BaseGizmos/GizmoViewContext.h"
 #include "BaseGizmos/PlanePositionGizmo.h"
 #include "BaseGizmos/AxisAngleGizmo.h"
-#include "BaseGizmos/TransformGizmo.h"
+#include "BaseGizmos/CombinedTransformGizmo.h"
 #include "BaseGizmos/RepositionableTransformGizmo.h"
 #include "BaseGizmos/ScalableSphereGizmo.h"
 #include "ContextObjectStore.h"
diff --git a/Engine/Source/Runtime/InteractiveToolsFramework/Public/BaseGizmos/CombinedTransformGizmo.h b/Engine/Source/Runtime/InteractiveToolsFramework/Public/BaseGizmos/CombinedTransformGizmo.h
new file mode 100644
--- /dev/null
+++ b/Engine/Source/Runtime/InteractiveToolsFramework/Public/BaseGizmos/CombinedTransformGizmo.h
@@ -0,0 +1,566 @@
+// Copyright Epic Games, Inc. All Rights Reserved.
+
+#pragma once
+
+#include "CoreMinimal.h"
+#include "InteractiveGizmo.h"
+#include "InteractiveGizmoBuilder.h"
+#include "InteractiveToolObjects.h"
+#include "InteractiveToolChange.h"
+#include "BaseGizmos/GizmoActor.h"
+#include "BaseGizmos/TransformProxy.h"
+
+#include "CombinedTransformGizmo.generated.h"
+
+class UInteractiveGizmoManager;
+class IGizmoAxisSource;
+class IGizmoTransformSource;
+class IGizmoStateTarget;
+class UGizmoConstantFrameAxisSource;
+class UGizmoComponentAxisSource;
+class UGizmoTransformChangeStateTarget;
+class UGizmoViewContext;
+class FTransformGizmoTransformChange;
+
+/**
+ * ACombinedTransformGizmoActor is an Actor type intended to be used with UCombinedTransformGizmo,
+ * as the in-scene visual representation of the Gizmo.
+ * 
+ * FCombinedTransformGizmoActorFactory returns an instance of this Actor type (or a subclass), and based on
+ * which Translate and Rotate UProperties are initialized, will associate those Components
+ * with UInteractiveGizmo's that implement Axis Translation, Plane Translation, and Axis Rotation.
+ * 
+ * If a particular sub-Gizmo is not required, simply set that FProperty to null.
+ * 
+ * The static factory method ::ConstructDefault3AxisGizmo() creates and initializes an 
+ * Actor suitable for use in a standard 3-axis Transformation Gizmo.
+ */
+UCLASS(Transient)
+class INTERACTIVETOOLSFRAMEWORK_API ACombinedTransformGizmoActor : public AGizmoActor
+{
+	GENERATED_BODY()
+public:
+
+	ACombinedTransformGizmoActor();
+
+public:
+	//
+	// Translation Components
+	//
+
+	/** X Axis Translation Component */
+	UPROPERTY()
+	TObjectPtr<UPrimitiveComponent> TranslateX;
+
+	/** Y Axis Translation Component */
+	UPROPERTY()
+	TObjectPtr<UPrimitiveComponent> TranslateY;
+
+	/** Z Axis Translation Component */
+	UPROPERTY()
+	TObjectPtr<UPrimitiveComponent> TranslateZ;
+
+
+	/** YZ Plane Translation Component */
+	UPROPERTY()
+	TObjectPtr<UPrimitiveComponent> TranslateYZ;
+
+	/** XZ Plane Translation Component */
+	UPROPERTY()
+	TObjectPtr<UPrimitiveComponent> TranslateXZ;
+
+	/** XY Plane Translation Component */
+	UPROPERTY()
+	TObjectPtr<UPrimitiveComponent> TranslateXY;
+
+	//
+	// Rotation Components
+	//
+
+	/** X Axis Rotation Component */
+	UPROPERTY()
+	TObjectPtr<UPrimitiveComponent> RotateX;
+
+	/** Y Axis Rotation Component */
+	UPROPERTY()
+	TObjectPtr<UPrimitiveComponent> RotateY;
+
+	/** Z Axis Rotation Component */
+	UPROPERTY()
+	TObjectPtr<UPrimitiveComponent> RotateZ;
+
+	//
+	// Scaling Components
+	//
+
+	/** Uniform Scale Component */
+	UPROPERTY()
+	TObjectPtr<UPrimitiveComponent> UniformScale;
+
+
+	/** X Axis Scale Component */
+	UPROPERTY()
+	TObjectPtr<UPrimitiveComponent> AxisScaleX;
+
+	/** Y Axis Scale Component */
+	UPROPERTY()
+	TObjectPtr<UPrimitiveComponent> AxisScaleY;
+
+	/** Z Axis Scale Component */
+	UPROPERTY()
+	TObjectPtr<UPrimitiveComponent> AxisScaleZ;
+
+
+	/** YZ Plane Scale Component */
+	UPROPERTY()
+	TObjectPtr<UPrimitiveComponent> PlaneScaleYZ;
+
+	/** XZ Plane Scale Component */
+	UPROPERTY()
+	TObjectPtr<UPrimitiveComponent> PlaneScaleXZ;
+
+	/** XY Plane Scale Component */
+	UPROPERTY()
+	TObjectPtr<UPrimitiveComponent> PlaneScaleXY;
+
+
+
+public:
+	/**
+	 * Create a new instance of ACombinedTransformGizmoActor and populate the various
+	 * sub-components with standard GizmoXComponent instances suitable for a 3-axis transformer Gizmo
+	 */
+	static ACombinedTransformGizmoActor* ConstructDefault3AxisGizmo(
+		UWorld* World, UGizmoViewContext* GizmoViewContext
+	);
+
+	/**
+	 * Create a new instance of ACombinedTransformGizmoActor. Populate the sub-components 
+	 * specified by Elements with standard GizmoXComponent instances suitable for a 3-axis transformer Gizmo
+	 */
+	static ACombinedTransformGizmoActor* ConstructCustom3AxisGizmo(
+		UWorld* World, UGizmoViewContext* GizmoViewContext,
+		ETransformGizmoSubElements Elements
+	);
+};
+
+
+
+
+
+/**
+ * FCombinedTransformGizmoActorFactory creates new instances of ACombinedTransformGizmoActor which
+ * are used by UCombinedTransformGizmo to implement 3D transformation Gizmos. 
+ * An instance of FCombinedTransformGizmoActorFactory is passed to UCombinedTransformGizmo
+ * (by way of UCombinedTransformGizmoBuilder), which then calls CreateNewGizmoActor()
+ * to spawn new Gizmo Actors.
+ * 
+ * By default CreateNewGizmoActor() returns a default Gizmo Actor suitable for
+ * a three-axis transformation Gizmo, override this function to customize
+ * the Actor sub-elements.
+ */
+class INTERACTIVETOOLSFRAMEWORK_API FCombinedTransformGizmoActorFactory
+{
+public:
+	FCombinedTransformGizmoActorFactory(UGizmoViewContext* GizmoViewContextIn)
+		: GizmoViewContext(GizmoViewContextIn)
+	{
+	}
+
+	/** Only these members of the ACombinedTransformGizmoActor gizmo will be initialized */
+	ETransformGizmoSubElements EnableElements =
+		ETransformGizmoSubElements::TranslateAllAxes |
+		ETransformGizmoSubElements::TranslateAllPlanes |
+		ETransformGizmoSubElements::RotateAllAxes |
+		ETransformGizmoSubElements::ScaleAllAxes |
+		ETransformGizmoSubElements::ScaleAllPlanes | 
+		ETransformGizmoSubElements::ScaleUniform;
+
+	/**
+	 * @param World the UWorld to create the new Actor in
+	 * @return new ACombinedTransformGizmoActor instance with members initialized with Components suitable for a transformation Gizmo
+	 */
+	virtual ACombinedTransformGizmoActor* CreateNewGizmoActor(UWorld* World) const;
+
+protected:
+	/**
+	 * The default gizmos that we use need to have the current view information stored for them via
+	 * the ITF context store so that they can figure out how big they are for hit testing, so this
+	 * pointer needs to be set (and kept alive elsewhere) for the actor factory to work properly.
+	 */
+	UGizmoViewContext* GizmoViewContext = nullptr;
+};
+
+
+
+
+
+
+UCLASS()
+class INTERACTIVETOOLSFRAMEWORK_API UCombinedTransformGizmoBuilder : public UInteractiveGizmoBuilder
+{
+	GENERATED_BODY()
+
+public:
+
+	/**
+	 * strings identifing GizmoBuilders already registered with GizmoManager. These builders will be used
+	 * to spawn the various sub-gizmos
+	 */
+	FString AxisPositionBuilderIdentifier;
+	FString PlanePositionBuilderIdentifier;
+	FString AxisAngleBuilderIdentifier;
+
+	/**
+	 * If set, this Actor Builder will be passed to UCombinedTransformGizmo instances.
+	 * Otherwise new instances of the base FCombinedTransformGizmoActorFactory are created internally.
+	 */
+	TSharedPtr<FCombinedTransformGizmoActorFactory> GizmoActorBuilder;
+
+	/**
+	 * If set, this hover function will be passed to UCombinedTransformGizmo instances to use instead of the default.
+	 * Hover is complicated for UCombinedTransformGizmo because all it knows about the different gizmo scene elements
+	 * is that they are UPrimitiveComponent (coming from the ACombinedTransformGizmoActor). The default hover
+	 * function implementation is to try casting to UGizmoBaseComponent and calling ::UpdateHoverState().
+	 * If you are using different Components that do not subclass UGizmoBaseComponent, and you want hover to 
+	 * work, you will need to provide a different hover update function.
+	 */
+	TFunction<void(UPrimitiveComponent*, bool)> UpdateHoverFunction;
+
+	/**
+	 * If set, this coord-system function will be passed to UCombinedTransformGizmo instances to use instead
+	 * of the default UpdateCoordSystemFunction. By default the UCombinedTransformGizmo will query the external Context
+	 * to ask whether it should be using world or local coordinate system. Then the default UpdateCoordSystemFunction
+	 * will try casting to UGizmoBaseCmponent and passing that info on via UpdateWorldLocalState();
+	 * If you are using different Components that do not subclass UGizmoBaseComponent, and you want the coord system
+	 * to be configurable, you will need to provide a different update function.
+	 */
+	TFunction<void(UPrimitiveComponent*, EToolContextCoordinateSystem)> UpdateCoordSystemFunction;
+
+
+	virtual UInteractiveGizmo* BuildGizmo(const FToolBuilderState& SceneState) const override;
+};
+
+
+/**
+ * UCombinedTransformGizmo provides standard Transformation Gizmo interactions,
+ * applied to a UTransformProxy target object. By default the Gizmo will be
+ * a standard XYZ translate/rotate Gizmo (axis and plane translation).
+ * 
+ * The in-scene representation of the Gizmo is a ACombinedTransformGizmoActor (or subclass).
+ * This Actor has FProperty members for the various sub-widgets, each as a separate Component.
+ * Any particular sub-widget of the Gizmo can be disabled by setting the respective
+ * Actor Component to null. 
+ * 
+ * So, to create non-standard variants of the Transform Gizmo, set a new GizmoActorBuilder 
+ * in the UCombinedTransformGizmoBuilder registered with the GizmoManager. Return
+ * a suitably-configured GizmoActor and everything else will be handled automatically.
+ * 
+ */
+UCLASS()
+class INTERACTIVETOOLSFRAMEWORK_API UCombinedTransformGizmo : public UInteractiveGizmo
+{
+	GENERATED_BODY()
+
+public:
+
+	virtual void SetWorld(UWorld* World);
+	virtual void SetGizmoActorBuilder(TSharedPtr<FCombinedTransformGizmoActorFactory> Builder);
+	virtual void SetSubGizmoBuilderIdentifiers(FString AxisPositionBuilderIdentifier, FString PlanePositionBuilderIdentifier, FString AxisAngleBuilderIdentifier);
+	virtual void SetUpdateHoverFunction(TFunction<void(UPrimitiveComponent*, bool)> HoverFunction);
+	virtual void SetUpdateCoordSystemFunction(TFunction<void(UPrimitiveComponent*, EToolContextCoordinateSystem)> CoordSysFunction);
+	
+	/**
+	 * If used, binds alignment functions to the sub gizmos that they can use to align to geometry in the scene. 
+	 * Specifically, translation and rotation gizmos will check ShouldAlignDestination() to see if they should
+	 * use the custom ray caster (this allows the behavior to respond to modifier key presses, for instance),
+	 * and then use DestinationAlignmentRayCaster() to find a point to align to.
+	 * Subgizmos align to the point in different ways, usually by projecting onto the axis or plane that they
+	 * operate in.
+	 */
+	virtual void SetWorldAlignmentFunctions(
+		TUniqueFunction<bool()>&& ShouldAlignDestination,
+		TUniqueFunction<bool(const FRay&, FVector&)>&& DestinationAlignmentRayCaster
+		);
+
+	/**
+	 * By default, non-uniform scaling handles appear (assuming they exist in the gizmo to begin with), 
+	 * when CurrentCoordinateSystem == EToolContextCoordinateSystem::Local, since components can only be
+	 * locally scaled. However, this can be changed to a custom check here, perhaps to hide them in extra
+	 * conditions or to always show them (if the gizmo is not scaling a component).
+	 */
+	virtual void SetIsNonUniformScaleAllowedFunction(
+		TUniqueFunction<bool()>&& IsNonUniformScaleAllowed
+	);
+
+	/**
+	 * By default, the nonuniform scale components can scale negatively. However, they can be made to clamp
+	 * to zero instead by passing true here. This is useful for using the gizmo to flatten geometry.
+	 *
+	 * TODO: Should this affect uniform scaling too?
+	 */
+	virtual void SetDisallowNegativeScaling(bool bDisallow);
+
+	// UInteractiveGizmo overrides
+	virtual void Setup() override;
+	virtual void Shutdown() override;
+	virtual void Tick(float DeltaTime) override;
+
+
+	/**
+	 * Set the active target object for the Gizmo
+	 * @param Target active target
+	 * @param TransactionProvider optional IToolContextTransactionProvider implementation to use - by default uses GizmoManager
+	 */
+	virtual void SetActiveTarget(UTransformProxy* Target, IToolContextTransactionProvider* TransactionProvider = nullptr);
+
+	/**
+	 * Clear the active target object for the Gizmo
+	 */
+	virtual void ClearActiveTarget();
+
+	/** The active target object for the Gizmo */
+	UPROPERTY()
+	TObjectPtr<UTransformProxy> ActiveTarget;
+
+	/**
+	 * @return the internal GizmoActor used by the Gizmo
+	 */
+	ACombinedTransformGizmoActor* GetGizmoActor() const { return GizmoActor; }
+
+	/**
+	 * @return current transform of Gizmo
+	 */
+	FTransform GetGizmoTransform() const;
+
+	/**
+	 * Repositions the gizmo without issuing undo/redo changes, triggering callbacks, 
+	 * or moving any components. Useful for resetting the gizmo to a new location without
+	 * it being viewed as a gizmo manipulation.
+	 * @param bKeepGizmoUnscaled If true, the scale component of NewTransform is passed through to the target but gizmo scale is set to 1
+	 */
+	void ReinitializeGizmoTransform(const FTransform& NewTransform, bool bKeepGizmoUnscaled = true);
+
+	/**
+	 * Set a new position for the Gizmo. This is done via the same mechanisms as the sub-gizmos,
+	 * so it generates the same Change/Modify() events, and hence works with Undo/Redo
+	 * @param bKeepGizmoUnscaled If true, the scale component of NewTransform is passed through to the target but gizmo scale is set to 1
+	 */
+	virtual void SetNewGizmoTransform(const FTransform& NewTransform, bool bKeepGizmoUnscaled = true);
+
+	/**
+	 * Explicitly set the child scale. Mainly useful to "reset" the child scale to (1,1,1) when re-using Gizmo across multiple transform actions.
+	 * @warning does not generate change/modify events!!
+	 */
+	virtual void SetNewChildScale(const FVector& NewChildScale);
+
+	/**
+	 * Set visibility for this Gizmo
+	 */
+	virtual void SetVisibility(bool bVisible);
+
+	/**
+	 * @return true if Gizmo is visible
+	 */
+	virtual bool IsVisible()
+	{
+		return GizmoActor && !GizmoActor->IsHidden();
+	}
+
+	/**
+	 * If true, then when using world frame, Axis and Plane translation snap to the world grid via the ContextQueriesAPI (in PositionSnapFunction)
+	 */
+	UPROPERTY()
+	bool bSnapToWorldGrid = true;
+
+	/**
+	 * Optional grid size which overrides the Context Grid
+	 */
+	UPROPERTY()
+	bool bGridSizeIsExplicit = false;
+	UPROPERTY()
+	FVector ExplicitGridSize;
+
+	/**
+	 * Optional grid size which overrides the Context Rotation Grid
+	 */
+	UPROPERTY()
+	bool bRotationGridSizeIsExplicit = false;
+	UPROPERTY()
+	FRotator ExplicitRotationGridSize;
+
+	/**
+	 * If true, then when using world frame, Axis and Plane translation snap to the world grid via the ContextQueriesAPI (in RotationSnapFunction)
+	 */
+	UPROPERTY()
+	bool bSnapToWorldRotGrid = true;
+
+	/**
+	 * Whether to use the World/Local coordinate system provided by the context via the ContextyQueriesAPI.
+	 */
+	UPROPERTY()
+	bool bUseContextCoordinateSystem = true;
+
+	/**
+	 * Current coordinate system in use. If bUseContextCoordinateSystem is true, this value will be updated internally every Tick()
+	 * by quering the ContextyQueriesAPI, otherwise the default is Local and the client can change it as necessary
+	 */
+	UPROPERTY()
+	EToolContextCoordinateSystem CurrentCoordinateSystem = EToolContextCoordinateSystem::Local;
+
+
+protected:
+	TSharedPtr<FCombinedTransformGizmoActorFactory> GizmoActorBuilder;
+
+	FString AxisPositionBuilderIdentifier;
+	FString PlanePositionBuilderIdentifier;
+	FString AxisAngleBuilderIdentifier;
+
+	// This function is called on each active GizmoActor Component to update it's hover state.
+	// If the Component is not a UGizmoBaseCmponent, the client needs to provide a different implementation
+	// of this function via the ToolBuilder
+	TFunction<void(UPrimitiveComponent*, bool)> UpdateHoverFunction;
+
+	// This function is called on each active GizmoActor Component to update it's coordinate system (eg world/local).
+	// If the Component is not a UGizmoBaseCmponent, the client needs to provide a different implementation
+	// of this function via the ToolBuilder
+	TFunction<void(UPrimitiveComponent*, EToolContextCoordinateSystem)> UpdateCoordSystemFunction;
+
+	/** List of current-active child components */
+	UPROPERTY()
+	TArray<TObjectPtr<UPrimitiveComponent>> ActiveComponents;
+
+	/** 
+	 * List of nonuniform scale components. Subset of of ActiveComponents. These are tracked separately so they can
+	 * be hidden when Gizmo is not configured to use local axes, because UE only supports local nonuniform scaling
+	 * on Components
+	 */
+	UPROPERTY()
+	TArray<TObjectPtr<UPrimitiveComponent>> NonuniformScaleComponents;
+
+	/** list of currently-active child gizmos */
+	UPROPERTY()
+	TArray<TObjectPtr<UInteractiveGizmo>> ActiveGizmos;
+
+	/** GizmoActors will be spawned in this World */
+	UWorld* World;
+
+	/** Current active GizmoActor that was spawned by this Gizmo. Will be destroyed when Gizmo is. */
+	ACombinedTransformGizmoActor* GizmoActor;
+
+	//
+	// Axis Sources
+	//
+
+
+	/** Axis that points towards camera, X/Y plane tangents aligned to right/up. Shared across Gizmos, and created internally during SetActiveTarget() */
+	UPROPERTY()
+	TObjectPtr<UGizmoConstantFrameAxisSource> CameraAxisSource;
+
+	// internal function that updates CameraAxisSource by getting current view state from GizmoManager
+	void UpdateCameraAxisSource();
+
+
+	/** X-axis source is shared across Gizmos, and created internally during SetActiveTarget() */
+	UPROPERTY()
+	TObjectPtr<UGizmoComponentAxisSource> AxisXSource;
+
+	/** Y-axis source is shared across Gizmos, and created internally during SetActiveTarget() */
+	UPROPERTY()
+	TObjectPtr<UGizmoComponentAxisSource> AxisYSource;
+
+	/** Z-axis source is shared across Gizmos, and created internally during SetActiveTarget() */
+	UPROPERTY()
+	TObjectPtr<UGizmoComponentAxisSource> AxisZSource;
+
+	//
+	// Scaling support. 
+	// UE Components only support scaling in local coordinates, so we have to create separate sources for that.
+	//
+
+	/** Local X-axis source (ie 1,0,0) is shared across Scale Gizmos, and created internally during SetActiveTarget() */
+	UPROPERTY()
+	TObjectPtr<UGizmoComponentAxisSource> UnitAxisXSource;
+
+	/** Y-axis source (ie 0,1,0) is shared across Scale Gizmos, and created internally during SetActiveTarget() */
+	UPROPERTY()
+	TObjectPtr<UGizmoComponentAxisSource> UnitAxisYSource;
+
+	/** Z-axis source (ie 0,0,1) is shared across Scale Gizmos, and created internally during SetActiveTarget() */
+	UPROPERTY()
+	TObjectPtr<UGizmoComponentAxisSource> UnitAxisZSource;
+
+
+	//
+	// Other Gizmo Components
+	//
+
+
+	/** 
+	 * State target is shared across gizmos, and created internally during SetActiveTarget(). 
+	 * Several FChange providers are registered with this StateTarget, including the UCombinedTransformGizmo
+	 * itself (IToolCommandChangeSource implementation above is called)
+	 */
+	UPROPERTY()
+	TObjectPtr<UGizmoTransformChangeStateTarget> StateTarget;
+
+	/**
+	 * These are used to let the translation subgizmos use raycasts into the scene to align the gizmo with scene geometry.
+	 * See comment for SetWorldAlignmentFunctions().
+	 */
+	TUniqueFunction<bool()> ShouldAlignDestination = []() { return false; };
+	TUniqueFunction<bool(const FRay&, FVector&)> DestinationAlignmentRayCaster = [](const FRay&, FVector&) {return false; };
+
+	TUniqueFunction<bool()> IsNonUniformScaleAllowed = [this]() { return CurrentCoordinateSystem == EToolContextCoordinateSystem::Local; };
+
+	bool bDisallowNegativeScaling = false;
+protected:
+
+
+	/** @return a new instance of the standard axis-translation Gizmo */
+	virtual UInteractiveGizmo* AddAxisTranslationGizmo(
+		UPrimitiveComponent* AxisComponent, USceneComponent* RootComponent,
+		IGizmoAxisSource* AxisSource,
+		IGizmoTransformSource* TransformSource, 
+		IGizmoStateTarget* StateTarget);
+
+	/** @return a new instance of the standard plane-translation Gizmo */
+	virtual UInteractiveGizmo* AddPlaneTranslationGizmo(
+		UPrimitiveComponent* AxisComponent, USceneComponent* RootComponent,
+		IGizmoAxisSource* AxisSource,
+		IGizmoTransformSource* TransformSource,
+		IGizmoStateTarget* StateTarget);
+
+	/** @return a new instance of the standard axis-rotation Gizmo */
+	virtual UInteractiveGizmo* AddAxisRotationGizmo(
+		UPrimitiveComponent* AxisComponent, USceneComponent* RootComponent,
+		IGizmoAxisSource* AxisSource,
+		IGizmoTransformSource* TransformSource,
+		IGizmoStateTarget* StateTarget);
+
+	/** @return a new instance of the standard axis-scaling Gizmo */
+	virtual UInteractiveGizmo* AddAxisScaleGizmo(
+		UPrimitiveComponent* AxisComponent, USceneComponent* RootComponent,
+		IGizmoAxisSource* GizmoAxisSource, IGizmoAxisSource* ParameterAxisSource,
+		IGizmoTransformSource* TransformSource,
+		IGizmoStateTarget* StateTarget);
+
+	/** @return a new instance of the standard plane-scaling Gizmo */
+	virtual UInteractiveGizmo* AddPlaneScaleGizmo(
+		UPrimitiveComponent* AxisComponent, USceneComponent* RootComponent,
+		IGizmoAxisSource* GizmoAxisSource, IGizmoAxisSource* ParameterAxisSource,
+		IGizmoTransformSource* TransformSource,
+		IGizmoStateTarget* StateTarget);
+
+	/** @return a new instance of the standard plane-scaling Gizmo */
+	virtual UInteractiveGizmo* AddUniformScaleGizmo(
+		UPrimitiveComponent* ScaleComponent, USceneComponent* RootComponent,
+		IGizmoAxisSource* GizmoAxisSource, IGizmoAxisSource* ParameterAxisSource,
+		IGizmoTransformSource* TransformSource,
+		IGizmoStateTarget* StateTarget);
+
+	// Axis and Plane TransformSources use this function to execute worldgrid snap queries
+	bool PositionSnapFunction(const FVector& WorldPosition, FVector& SnappedPositionOut) const;
+	FQuat RotationSnapFunction(const FQuat& DeltaRotation) const;
+
+};
diff --git a/Engine/Source/Runtime/InteractiveToolsFramework/Public/BaseGizmos/RepositionableTransformGizmo.h b/Engine/Source/Runtime/InteractiveToolsFramework/Public/BaseGizmos/RepositionableTransformGizmo.h
--- a/Engine/Source/Runtime/InteractiveToolsFramework/Public/BaseGizmos/RepositionableTransformGizmo.h
+++ b/Engine/Source/Runtime/InteractiveToolsFramework/Public/BaseGizmos/RepositionableTransformGizmo.h
@@ -4,7 +4,7 @@
 
 #include "CoreMinimal.h"
 
-#include "BaseGizmos/TransformGizmo.h"
+#include "BaseGizmos/CombinedTransformGizmo.h"
 #include "InteractiveGizmoBuilder.h"
 #include "ToolContextInterfaces.h"
 
diff --git a/Engine/Source/Runtime/InteractiveToolsFramework/Public/BaseGizmos/TransformGizmo.h b/Engine/Source/Runtime/InteractiveToolsFramework/Public/BaseGizmos/TransformGizmo.h
deleted file mode 100644
diff --git a/Engine/Source/Runtime/InteractiveToolsFramework/Public/BaseGizmos/TransformGizmoUtil.h b/Engine/Source/Runtime/InteractiveToolsFramework/Public/BaseGizmos/TransformGizmoUtil.h
--- a/Engine/Source/Runtime/InteractiveToolsFramework/Public/BaseGizmos/TransformGizmoUtil.h
+++ b/Engine/Source/Runtime/InteractiveToolsFramework/Public/BaseGizmos/TransformGizmoUtil.h
@@ -4,7 +4,7 @@
 
 #include "CoreMinimal.h"
 #include "InteractiveGizmo.h"
-#include "BaseGizmos/TransformGizmo.h"
+#include "BaseGizmos/CombinedTransformGizmo.h"
 
 #include "TransformGizmoUtil.generated.h"
 
