MeshModelingTools: Added ULocalInputBehaviorSource to allow tools to give different behaviors to separate input routers, and removed direct RegisterBehavior function from UInputRouter to force use of behavior sources. Added  ULocalMouseHoverBehavior.

#rnx

diff --git a/Engine/Source/Runtime/InteractiveToolsFramework/Private/InputRouter.cpp b/Engine/Source/Runtime/InteractiveToolsFramework/Private/InputRouter.cpp
--- a/Engine/Source/Runtime/InteractiveToolsFramework/Private/InputRouter.cpp
+++ b/Engine/Source/Runtime/InteractiveToolsFramework/Private/InputRouter.cpp
@@ -48,12 +48,6 @@ void UInputRouter::DeregisterSource(IInputBehaviorSource* Source)
 }
 
 
-void UInputRouter::RegisterBehavior(UInputBehavior* Behavior, void* Source, const FString& GroupName)
-{
-	ActiveInputBehaviors->Add(Behavior, Source, GroupName);
-}
-
-
 bool UInputRouter::PostInputEvent(const FInputDeviceState& Input)
 {
 	if (ActiveInputBehaviors->IsEmpty())
diff --git a/Engine/Source/Runtime/InteractiveToolsFramework/Public/BaseBehaviors/MouseHoverBehavior.h b/Engine/Source/Runtime/InteractiveToolsFramework/Public/BaseBehaviors/MouseHoverBehavior.h
--- a/Engine/Source/Runtime/InteractiveToolsFramework/Public/BaseBehaviors/MouseHoverBehavior.h
+++ b/Engine/Source/Runtime/InteractiveToolsFramework/Public/BaseBehaviors/MouseHoverBehavior.h
@@ -45,3 +45,56 @@ protected:
 };
 
 
+/**
+ * An implementation of UMouseHoverBehavior that also implements IHoverBehaviorTarget directly, via a set 
+ * of local lambda functions. To use/customize this class, the client replaces the lambda functions with their own.
+ * This avoids having to create a separate IHoverBehaviorTarget implementation for trivial use-cases.
+ */
+UCLASS()
+class INTERACTIVETOOLSFRAMEWORK_API ULocalMouseHoverBehavior : public UMouseHoverBehavior, public IHoverBehaviorTarget
+{
+	GENERATED_BODY()
+protected:
+	using UMouseHoverBehavior::Initialize;
+
+public:
+	/** Call this to initialize the class */
+	virtual void Initialize()
+	{
+		this->Initialize(this);
+	}
+
+	/** lambda implementation of BeginHoverSequenceHitTest */
+	TUniqueFunction<FInputRayHit(const FInputDeviceRay& PressPos)> BeginHitTestFunc = [](const FInputDeviceRay&) { return FInputRayHit(); };
+
+	/** lambda implementation of OnBeginHover */
+	TUniqueFunction<void(const FInputDeviceRay& PressPos)> OnBeginHoverFunc = [](const FInputDeviceRay&) {};
+
+	/** lambda implementation of OnUpdateHover */
+	TUniqueFunction<bool(const FInputDeviceRay& PressPos)> OnUpdateHoverFunc = [](const FInputDeviceRay&) { return false; };
+
+	/** lambda implementation of OnEndHover */
+	TUniqueFunction<void()> OnEndHoverFunc = []() {};
+
+public:
+	// IHoverBehaviorTarget implementation
+	virtual FInputRayHit BeginHoverSequenceHitTest(const FInputDeviceRay& PressPos) override
+	{
+		return BeginHitTestFunc(PressPos);
+	}
+
+	virtual void OnBeginHover(const FInputDeviceRay& DevicePos) override
+	{
+		OnBeginHoverFunc(DevicePos);
+	}
+
+	virtual bool OnUpdateHover(const FInputDeviceRay& DevicePos) override
+	{
+		return OnUpdateHoverFunc(DevicePos);
+	}
+
+	virtual void OnEndHover() override
+	{
+		OnEndHoverFunc();
+	}
+};
diff --git a/Engine/Source/Runtime/InteractiveToolsFramework/Public/InputBehaviorSet.h b/Engine/Source/Runtime/InteractiveToolsFramework/Public/InputBehaviorSet.h
--- a/Engine/Source/Runtime/InteractiveToolsFramework/Public/InputBehaviorSet.h
+++ b/Engine/Source/Runtime/InteractiveToolsFramework/Public/InputBehaviorSet.h
@@ -171,3 +171,21 @@ public:
 	virtual const UInputBehaviorSet* GetInputBehaviors() const = 0;
 };
 
+
+/**
+ * An implementation of IInputBehaviorSource that forwards to a user provided-lambda, to allow
+ * a tool to supply a behavior source different from the one it is implementing itself. Useful,
+ * for instance, when a tool wants to supply different behaviors to separate input routers.
+ */
+UCLASS(Transient)
+class INTERACTIVETOOLSFRAMEWORK_API ULocalInputBehaviorSource : public UObject, public IInputBehaviorSource
+{
+	GENERATED_BODY()
+
+public:
+
+	TUniqueFunction<const UInputBehaviorSet* ()> GetInputBehaviorsFunc = []() { return nullptr; };
+
+	// IInputBehaviorSource
+	virtual const UInputBehaviorSet* GetInputBehaviors() const override { return GetInputBehaviorsFunc(); };
+};
diff --git a/Engine/Source/Runtime/InteractiveToolsFramework/Public/InputRouter.h b/Engine/Source/Runtime/InteractiveToolsFramework/Public/InputRouter.h
--- a/Engine/Source/Runtime/InteractiveToolsFramework/Public/InputRouter.h
+++ b/Engine/Source/Runtime/InteractiveToolsFramework/Public/InputRouter.h
@@ -47,10 +47,6 @@ public:
 	/** Remove Behaviors from this Source from the active set */
 	virtual void DeregisterSource(IInputBehaviorSource* Source);
 
-	/** Add a new InputBehavior to the active behavior set, with optional Source and GroupName */
-	virtual void RegisterBehavior(UInputBehavior* Behavior, void* Source = nullptr, const FString& GroupName = "");
-
-
 	/** Insert a new input event which is used to check for new captures, or forwarded to active capture */
 	virtual bool PostInputEvent(const FInputDeviceState& Input);
 
