MeshModelingTools: Make rotation not try to (incorrectly) snap in local mode.

#rnx

diff --git a/Engine/Source/Runtime/InteractiveToolsFramework/Private/BaseGizmos/CombinedTransformGizmo.cpp b/Engine/Source/Runtime/InteractiveToolsFramework/Private/BaseGizmos/CombinedTransformGizmo.cpp
--- a/Engine/Source/Runtime/InteractiveToolsFramework/Private/BaseGizmos/CombinedTransformGizmo.cpp
+++ b/Engine/Source/Runtime/InteractiveToolsFramework/Private/BaseGizmos/CombinedTransformGizmo.cpp
@@ -963,26 +963,41 @@ FQuat UCombinedTransformGizmo::RotationSnapFunction(const FQuat& DeltaRotation)
 	FQuat SnappedDeltaRotation = DeltaRotation;
 
 	// only snap if we want snapping obvs
-	if (bSnapToWorldRotGrid)
+	if (!bSnapToWorldGrid)
 	{
-		USceneSnappingManager* SnapManager = USceneSnappingManager::Find(GetGizmoManager());
-		if ( SnapManager )
-		{
-			FSceneSnapQueryRequest Request;
-			Request.RequestType   = ESceneSnapQueryType::Rotation;
-			Request.TargetTypes   = ESceneSnapQueryTargetType::Grid;
-			Request.DeltaRotation = DeltaRotation;
-			if ( bRotationGridSizeIsExplicit )
-			{
-				Request.RotGridSize = ExplicitRotationGridSize;
-			}
-			TArray<FSceneSnapQueryResult> Results;
-			if (SnapManager->ExecuteSceneSnapQuery(Request, Results))
-			{
-				SnappedDeltaRotation = Results[0].DeltaRotation;
-			};
-		}
+		return SnappedDeltaRotation;
 	}
+
+	// To match our position snapping behavior, only snap when using world axes.
+	// Note that if we someday want to snap in local mode, this function will need further
+	// changing because the quaternion is given and snapped in world space, whereas we
+	// would want to snap it relative to the local frame start orientation.
+	EToolContextCoordinateSystem CoordSystem = GetGizmoManager()->GetContextQueriesAPI()->GetCurrentCoordinateSystem();
+	if (CoordSystem != EToolContextCoordinateSystem::World)
+	{
+		return SnappedDeltaRotation;
+	}
+
+	USceneSnappingManager* SnapManager = USceneSnappingManager::Find(GetGizmoManager());
+	if (!SnapManager)
+	{
+		return SnappedDeltaRotation;
+	}
+
+	FSceneSnapQueryRequest Request;
+	Request.RequestType   = ESceneSnapQueryType::Rotation;
+	Request.TargetTypes   = ESceneSnapQueryTargetType::Grid;
+	Request.DeltaRotation = DeltaRotation;
+	if ( bRotationGridSizeIsExplicit )
+	{
+		Request.RotGridSize = ExplicitRotationGridSize;
+	}
+	TArray<FSceneSnapQueryResult> Results;
+	if (SnapManager->ExecuteSceneSnapQuery(Request, Results))
+	{
+		SnappedDeltaRotation = Results[0].DeltaRotation;
+	};
+
 	return SnappedDeltaRotation;
 }
 
