Check if ToolManager is null before accessing it to close all tools, to avoid crash if cleanup functions are called at the wrong time

#rnx

diff --git a/Engine/Source/Runtime/InteractiveToolsFramework/Private/InteractiveToolsContext.cpp b/Engine/Source/Runtime/InteractiveToolsFramework/Private/InteractiveToolsContext.cpp
--- a/Engine/Source/Runtime/InteractiveToolsFramework/Private/InteractiveToolsContext.cpp
+++ b/Engine/Source/Runtime/InteractiveToolsFramework/Private/InteractiveToolsContext.cpp
@@ -159,6 +159,10 @@ void UInteractiveToolsContext::DeactivateActiveTool(EToolSide WhichSide, EToolSh
 
 void UInteractiveToolsContext::DeactivateAllActiveTools(EToolShutdownType ShutdownType)
 {
+	if (!ensureMsgf(ToolManager, TEXT("Trying to deactivate all tools on a ToolManager that has already been deleted.")))
+	{
+		return;
+	}
 	auto DeactivateTool = [this, ShutdownType](EToolSide WhichSide) {
 		if (ToolManager->HasActiveTool(WhichSide))
 		{
