InteractiveToolsFramework: Consolidate InteractiveToolWithToolTargetsBuilder into InteractiveToolBuilder.h

#rnx

diff --git a/Engine/Plugins/Experimental/MeshModelingToolsetExp/Source/MeshModelingToolsExp/Public/CubeGridTool.h b/Engine/Plugins/Experimental/MeshModelingToolsetExp/Source/MeshModelingToolsExp/Public/CubeGridTool.h
--- a/Engine/Plugins/Experimental/MeshModelingToolsetExp/Source/MeshModelingToolsExp/Public/CubeGridTool.h
+++ b/Engine/Plugins/Experimental/MeshModelingToolsetExp/Source/MeshModelingToolsExp/Public/CubeGridTool.h
@@ -10,7 +10,7 @@
 #include "FrameTypes.h"
 #include "GeometryBase.h"
 #include "InteractiveTool.h"
-#include "InteractiveToolWithToolTargetsBuilder.h"
+#include "InteractiveToolBuilder.h"
 #include "InteractiveToolQueryInterfaces.h" // IInteractiveToolNestedAcceptCancelAPI
 #include "Mechanics/CubeGrid.h"
 #include "ModelingOperators.h"
diff --git a/Engine/Plugins/Runtime/MeshModelingToolset/Source/ModelingComponents/Public/BaseTools/BaseMeshProcessingTool.h b/Engine/Plugins/Runtime/MeshModelingToolset/Source/ModelingComponents/Public/BaseTools/BaseMeshProcessingTool.h
--- a/Engine/Plugins/Runtime/MeshModelingToolset/Source/ModelingComponents/Public/BaseTools/BaseMeshProcessingTool.h
+++ b/Engine/Plugins/Runtime/MeshModelingToolset/Source/ModelingComponents/Public/BaseTools/BaseMeshProcessingTool.h
@@ -4,7 +4,7 @@
 
 #include "CoreMinimal.h"
 #include "SingleSelectionTool.h"
-#include "InteractiveToolWithToolTargetsBuilder.h"
+#include "InteractiveToolBuilder.h"
 #include "DynamicMesh/DynamicMesh3.h"
 #include "WeightMapTypes.h"
 #include "MeshOpPreviewHelpers.h"
diff --git a/Engine/Plugins/Runtime/MeshModelingToolset/Source/ModelingComponents/Public/BaseTools/MultiSelectionMeshEditingTool.h b/Engine/Plugins/Runtime/MeshModelingToolset/Source/ModelingComponents/Public/BaseTools/MultiSelectionMeshEditingTool.h
--- a/Engine/Plugins/Runtime/MeshModelingToolset/Source/ModelingComponents/Public/BaseTools/MultiSelectionMeshEditingTool.h
+++ b/Engine/Plugins/Runtime/MeshModelingToolset/Source/ModelingComponents/Public/BaseTools/MultiSelectionMeshEditingTool.h
@@ -3,7 +3,7 @@
 #pragma once
 
 #include "MultiSelectionTool.h"
-#include "InteractiveToolWithToolTargetsBuilder.h"
+#include "InteractiveToolBuilder.h"
 #include "MultiSelectionMeshEditingTool.generated.h"
 
 class UMultiSelectionMeshEditingTool;
diff --git a/Engine/Plugins/Runtime/MeshModelingToolset/Source/ModelingComponents/Public/BaseTools/SingleSelectionMeshEditingTool.h b/Engine/Plugins/Runtime/MeshModelingToolset/Source/ModelingComponents/Public/BaseTools/SingleSelectionMeshEditingTool.h
--- a/Engine/Plugins/Runtime/MeshModelingToolset/Source/ModelingComponents/Public/BaseTools/SingleSelectionMeshEditingTool.h
+++ b/Engine/Plugins/Runtime/MeshModelingToolset/Source/ModelingComponents/Public/BaseTools/SingleSelectionMeshEditingTool.h
@@ -3,7 +3,7 @@
 #pragma once
 
 #include "SingleSelectionTool.h"
-#include "InteractiveToolWithToolTargetsBuilder.h"
+#include "InteractiveToolBuilder.h"
 #include "SingleSelectionMeshEditingTool.generated.h"
 
 class USingleSelectionMeshEditingTool;
diff --git a/Engine/Source/Runtime/InteractiveToolsFramework/Public/BaseTools/MeshSurfacePointTool.h b/Engine/Source/Runtime/InteractiveToolsFramework/Public/BaseTools/MeshSurfacePointTool.h
--- a/Engine/Source/Runtime/InteractiveToolsFramework/Public/BaseTools/MeshSurfacePointTool.h
+++ b/Engine/Source/Runtime/InteractiveToolsFramework/Public/BaseTools/MeshSurfacePointTool.h
@@ -4,7 +4,7 @@
 
 #include "CoreMinimal.h"
 #include "SingleSelectionTool.h"
-#include "InteractiveToolWithToolTargetsBuilder.h"
+#include "InteractiveToolBuilder.h"
 #include "InputBehaviorSet.h"
 #include "BaseBehaviors/BehaviorTargetInterfaces.h"
 #include "BaseBehaviors/AnyButtonInputBehavior.h"
diff --git a/Engine/Source/Runtime/InteractiveToolsFramework/Public/InteractiveToolBuilder.h b/Engine/Source/Runtime/InteractiveToolsFramework/Public/InteractiveToolBuilder.h
--- a/Engine/Source/Runtime/InteractiveToolsFramework/Public/InteractiveToolBuilder.h
+++ b/Engine/Source/Runtime/InteractiveToolsFramework/Public/InteractiveToolBuilder.h
@@ -67,4 +67,27 @@ public:
 	{
 		return;
 	}
-};
\ No newline at end of file
+};
+
+/**
+ * A UInteractiveToolWithToolTargetsBuilder creates a new instance of an InteractiveTool that uses tool targets.
+ * See ToolTarget.h for more information on tool targets and their usage. This class defines the common
+ * interface(s) for defining the tool target requirements of the tool it builds.
+ * This is an abstract base class, you must subclass it in order to create your particular Tool instance.
+ */
+UCLASS(Transient, Abstract)
+class INTERACTIVETOOLSFRAMEWORK_API UInteractiveToolWithToolTargetsBuilder : public UInteractiveToolBuilder
+{
+	GENERATED_BODY()
+
+protected:
+	/**
+	 * Gives the target requirements of the associated tool. Usually, it is the tool builder
+	 * will use this function in CanBuildTool and BuildTool to find and create any necessary targets.
+	 */
+	virtual const FToolTargetTypeRequirements& GetTargetRequirements() const
+	{
+		static FToolTargetTypeRequirements TypeRequirements; // Default initialized to no requirements.
+		return TypeRequirements;
+	}
+};
diff --git a/Engine/Source/Runtime/InteractiveToolsFramework/Public/InteractiveToolWithToolTargetsBuilder.h b/Engine/Source/Runtime/InteractiveToolsFramework/Public/InteractiveToolWithToolTargetsBuilder.h
deleted file mode 100644
