InteractiveToolsFramework: Add UPrimitiveComponentToolTargetFactory.

#rnx

diff --git a/Engine/Source/Runtime/InteractiveToolsFramework/Private/ToolTargets/PrimitiveComponentToolTarget.cpp b/Engine/Source/Runtime/InteractiveToolsFramework/Private/ToolTargets/PrimitiveComponentToolTarget.cpp
--- a/Engine/Source/Runtime/InteractiveToolsFramework/Private/ToolTargets/PrimitiveComponentToolTarget.cpp
+++ b/Engine/Source/Runtime/InteractiveToolsFramework/Private/ToolTargets/PrimitiveComponentToolTarget.cpp
@@ -41,3 +41,26 @@ bool UPrimitiveComponentToolTarget::HitTestComponent(const FRay& WorldRay, FHitR
 	}
 	return false;
 }
+
+
+// Factory
+
+bool UPrimitiveComponentToolTargetFactory::CanBuildTarget(UObject* SourceObject, const FToolTargetTypeRequirements& Requirements) const
+{
+	UPrimitiveComponent* Component = Cast<UPrimitiveComponent>(SourceObject);
+	return Component 
+		&& IsValidChecked(Component)
+		&& !Component->IsUnreachable() 
+		&& Component->IsValidLowLevel() 
+		&& Requirements.AreSatisfiedBy(UPrimitiveComponentToolTarget::StaticClass());
+}
+
+UToolTarget* UPrimitiveComponentToolTargetFactory::BuildTarget(UObject* SourceObject, const FToolTargetTypeRequirements& Requirements)
+{
+	UPrimitiveComponentToolTarget* Target = NewObject<UPrimitiveComponentToolTarget>();
+	Target->Component = Cast<UPrimitiveComponent>(SourceObject);
+	check(Target->Component && Requirements.AreSatisfiedBy(Target));
+	return Target;
+}
+
+
diff --git a/Engine/Source/Runtime/InteractiveToolsFramework/Public/ToolTargets/PrimitiveComponentToolTarget.h b/Engine/Source/Runtime/InteractiveToolsFramework/Public/ToolTargets/PrimitiveComponentToolTarget.h
--- a/Engine/Source/Runtime/InteractiveToolsFramework/Public/ToolTargets/PrimitiveComponentToolTarget.h
+++ b/Engine/Source/Runtime/InteractiveToolsFramework/Public/ToolTargets/PrimitiveComponentToolTarget.h
@@ -10,10 +10,10 @@
 #include "PrimitiveComponentToolTarget.generated.h"
 
 /** 
- * An abstract tool target to share some reusable code for tool targets that are
+ * A tool target to share some reusable code for tool targets that are
  * backed by primitive components. 
  */
-UCLASS(Transient, Abstract)
+UCLASS(Transient)
 class INTERACTIVETOOLSFRAMEWORK_API UPrimitiveComponentToolTarget : public UToolTarget, public IPrimitiveComponentBackedTarget
 {
 	GENERATED_BODY()
@@ -28,5 +28,20 @@ public:
 	virtual bool HitTestComponent(const FRay& WorldRay, FHitResult& OutHit) const override;
 
 protected:
+	friend class UPrimitiveComponentToolTargetFactory;
+	
 	UPrimitiveComponent* Component;
-};
\ No newline at end of file
+};
+
+/**
+ * Factory for UPrimitiveComponentToolTarget to be used by the target manager.
+ */
+UCLASS(Transient)
+class INTERACTIVETOOLSFRAMEWORK_API UPrimitiveComponentToolTargetFactory : public UToolTargetFactory
+{
+	GENERATED_BODY()
+public:
+	virtual bool CanBuildTarget(UObject* SourceObject, const FToolTargetTypeRequirements& TargetTypeInfo) const override;
+	virtual UToolTarget* BuildTarget(UObject* SourceObject, const FToolTargetTypeRequirements& TargetTypeInfo) override;
+};
+
