Don't snap the transform gizmo on mouse-down, wait for the drag event.

#rnx

diff --git a/Engine/Source/Runtime/InteractiveToolsFramework/Private/BaseGizmos/AxisPositionGizmo.cpp b/Engine/Source/Runtime/InteractiveToolsFramework/Private/BaseGizmos/AxisPositionGizmo.cpp
--- a/Engine/Source/Runtime/InteractiveToolsFramework/Private/BaseGizmos/AxisPositionGizmo.cpp
+++ b/Engine/Source/Runtime/InteractiveToolsFramework/Private/BaseGizmos/AxisPositionGizmo.cpp
@@ -95,8 +95,6 @@ void UAxisPositionGizmo::OnClickPress(const FInputDeviceRay& PressPos)
 	{
 		StateTarget->BeginUpdate();
 	}
-
-	OnClickDrag(PressPos);
 }
 
 void UAxisPositionGizmo::OnClickDrag(const FInputDeviceRay& DragPos)
