ToolsFramework: change IsValidChecked to IsValid in several places where we want to handle failure gracefully
#rnx

diff --git a/Engine/Source/Runtime/InteractiveToolsFramework/Private/ToolTargets/PrimitiveComponentToolTarget.cpp b/Engine/Source/Runtime/InteractiveToolsFramework/Private/ToolTargets/PrimitiveComponentToolTarget.cpp
--- a/Engine/Source/Runtime/InteractiveToolsFramework/Private/ToolTargets/PrimitiveComponentToolTarget.cpp
+++ b/Engine/Source/Runtime/InteractiveToolsFramework/Private/ToolTargets/PrimitiveComponentToolTarget.cpp
@@ -6,7 +6,7 @@
 
 bool UPrimitiveComponentToolTarget::IsValid() const
 {
-	return Component && IsValidChecked(Component) && !Component->IsUnreachable() && Component->IsValidLowLevel();
+	return Component && ::IsValid(Component) && !Component->IsUnreachable() && Component->IsValidLowLevel();
 }
 
 UPrimitiveComponent* UPrimitiveComponentToolTarget::GetOwnerComponent() const
@@ -49,7 +49,7 @@ bool UPrimitiveComponentToolTargetFactory::CanBuildTarget(UObject* SourceObject,
 {
 	UPrimitiveComponent* Component = Cast<UPrimitiveComponent>(SourceObject);
 	return Component 
-		&& IsValidChecked(Component)
+		&& ::IsValid(Component)
 		&& !Component->IsUnreachable() 
 		&& Component->IsValidLowLevel() 
 		&& Requirements.AreSatisfiedBy(UPrimitiveComponentToolTarget::StaticClass());
