make brush tools notify property sets of updates done by hotkey actions (to brush size/strength), so the reset arrow correctly updates in the details panels

diff --git a/Engine/Plugins/Experimental/MeshModelingToolsetExp/Source/MeshModelingToolsExp/Private/DynamicMeshSculptTool.cpp b/Engine/Plugins/Experimental/MeshModelingToolsetExp/Source/MeshModelingToolsExp/Private/DynamicMeshSculptTool.cpp
--- a/Engine/Plugins/Experimental/MeshModelingToolsetExp/Source/MeshModelingToolsExp/Private/DynamicMeshSculptTool.cpp
+++ b/Engine/Plugins/Experimental/MeshModelingToolsetExp/Source/MeshModelingToolsExp/Private/DynamicMeshSculptTool.cpp
@@ -2197,31 +2197,37 @@ UPreviewMesh* UDynamicMeshSculptTool::MakeDefaultSphereMesh(UObject* Parent, UWo
 void UDynamicMeshSculptTool::IncreaseBrushRadiusAction()
 {	
 	BrushProperties->BrushSize.IncreaseRadius(false);
+	NotifyOfPropertyChangeByTool(BrushProperties);
 }
 
 void UDynamicMeshSculptTool::DecreaseBrushRadiusAction()
 {
 	BrushProperties->BrushSize.DecreaseRadius(false);
+	NotifyOfPropertyChangeByTool(BrushProperties);
 }
 
 void UDynamicMeshSculptTool::IncreaseBrushRadiusSmallStepAction()
 {
 	BrushProperties->BrushSize.IncreaseRadius(true);
+	NotifyOfPropertyChangeByTool(BrushProperties);
 }
 
 void UDynamicMeshSculptTool::DecreaseBrushRadiusSmallStepAction()
 {
 	BrushProperties->BrushSize.DecreaseRadius(true);
+	NotifyOfPropertyChangeByTool(BrushProperties);
 }
 
 void UDynamicMeshSculptTool::IncreaseBrushSpeedAction()
 {
 	SculptProperties->PrimaryBrushSpeed = FMath::Clamp(SculptProperties->PrimaryBrushSpeed + 0.05f, 0.0f, 1.0f);
+	NotifyOfPropertyChangeByTool(SculptProperties);
 }
 
 void UDynamicMeshSculptTool::DecreaseBrushSpeedAction()
 {
 	SculptProperties->PrimaryBrushSpeed = FMath::Clamp(SculptProperties->PrimaryBrushSpeed - 0.05f, 0.0f, 1.0f);
+	NotifyOfPropertyChangeByTool(SculptProperties);
 }
 
 
diff --git a/Engine/Plugins/Experimental/MeshModelingToolsetExp/Source/MeshModelingToolsExp/Private/MeshGroupPaintTool.cpp b/Engine/Plugins/Experimental/MeshModelingToolsetExp/Source/MeshModelingToolsExp/Private/MeshGroupPaintTool.cpp
--- a/Engine/Plugins/Experimental/MeshModelingToolsetExp/Source/MeshModelingToolsExp/Private/MeshGroupPaintTool.cpp
+++ b/Engine/Plugins/Experimental/MeshModelingToolsetExp/Source/MeshModelingToolsExp/Private/MeshGroupPaintTool.cpp
@@ -325,24 +325,28 @@ void UMeshGroupPaintTool::IncreaseBrushRadiusAction()
 {
 	Super::IncreaseBrushRadiusAction();
 	FilterProperties->BrushSize = BrushProperties->BrushSize.AdaptiveSize;
+	NotifyOfPropertyChangeByTool(FilterProperties);
 }
 
 void UMeshGroupPaintTool::DecreaseBrushRadiusAction()
 {
 	Super::DecreaseBrushRadiusAction();
 	FilterProperties->BrushSize = BrushProperties->BrushSize.AdaptiveSize;
+	NotifyOfPropertyChangeByTool(FilterProperties);
 }
 
 void UMeshGroupPaintTool::IncreaseBrushRadiusSmallStepAction()
 {
 	Super::IncreaseBrushRadiusSmallStepAction();
 	FilterProperties->BrushSize = BrushProperties->BrushSize.AdaptiveSize;
+	NotifyOfPropertyChangeByTool(FilterProperties);
 }
 
 void UMeshGroupPaintTool::DecreaseBrushRadiusSmallStepAction()
 {
 	Super::DecreaseBrushRadiusSmallStepAction();
 	FilterProperties->BrushSize = BrushProperties->BrushSize.AdaptiveSize;
+	NotifyOfPropertyChangeByTool(FilterProperties);
 }
 
 
diff --git a/Engine/Plugins/Experimental/MeshModelingToolsetExp/Source/MeshModelingToolsExp/Private/MeshVertexSculptTool.cpp b/Engine/Plugins/Experimental/MeshModelingToolsetExp/Source/MeshModelingToolsExp/Private/MeshVertexSculptTool.cpp
--- a/Engine/Plugins/Experimental/MeshModelingToolsetExp/Source/MeshModelingToolsExp/Private/MeshVertexSculptTool.cpp
+++ b/Engine/Plugins/Experimental/MeshModelingToolsetExp/Source/MeshModelingToolsExp/Private/MeshVertexSculptTool.cpp
@@ -1081,6 +1081,7 @@ void UMeshVertexSculptTool::IncreaseBrushSpeedAction()
 	float CurStrength = UseBrushOp->PropertySet->GetStrength();
 	float NewStrength = FMath::Clamp(CurStrength + 0.05f, 0.0f, 1.0f);
 	UseBrushOp->PropertySet->SetStrength(NewStrength);
+	NotifyOfPropertyChangeByTool(UseBrushOp->PropertySet.Get());
 }
 
 void UMeshVertexSculptTool::DecreaseBrushSpeedAction()
@@ -1089,6 +1090,7 @@ void UMeshVertexSculptTool::DecreaseBrushSpeedAction()
 	float CurStrength = UseBrushOp->PropertySet->GetStrength();
 	float NewStrength = FMath::Clamp(CurStrength - 0.05f, 0.0f, 1.0f);
 	UseBrushOp->PropertySet->SetStrength(NewStrength);
+	NotifyOfPropertyChangeByTool(UseBrushOp->PropertySet.Get());
 }
 
 void UMeshVertexSculptTool::UpdateBrushAlpha(UTexture2D* NewAlpha)
diff --git a/Engine/Plugins/Experimental/MeshModelingToolsetExp/Source/MeshModelingToolsExp/Private/Sculpting/MeshSculptToolBase.cpp b/Engine/Plugins/Experimental/MeshModelingToolsetExp/Source/MeshModelingToolsExp/Private/Sculpting/MeshSculptToolBase.cpp
--- a/Engine/Plugins/Experimental/MeshModelingToolsetExp/Source/MeshModelingToolsExp/Private/Sculpting/MeshSculptToolBase.cpp
+++ b/Engine/Plugins/Experimental/MeshModelingToolsetExp/Source/MeshModelingToolsExp/Private/Sculpting/MeshSculptToolBase.cpp
@@ -931,24 +931,28 @@ double UMeshSculptToolBase::GetCurrentBrushDepth()
 void UMeshSculptToolBase::IncreaseBrushRadiusAction()
 {
 	BrushProperties->BrushSize.IncreaseRadius(false);
+	NotifyOfPropertyChangeByTool(BrushProperties);
 	CalculateBrushRadius();
 }
 
 void UMeshSculptToolBase::DecreaseBrushRadiusAction()
 {
 	BrushProperties->BrushSize.DecreaseRadius(false);
+	NotifyOfPropertyChangeByTool(BrushProperties);
 	CalculateBrushRadius();
 }
 
 void UMeshSculptToolBase::IncreaseBrushRadiusSmallStepAction()
 {
 	BrushProperties->BrushSize.IncreaseRadius(true);
+	NotifyOfPropertyChangeByTool(BrushProperties);
 	CalculateBrushRadius();
 }
 
 void UMeshSculptToolBase::DecreaseBrushRadiusSmallStepAction()
 {
 	BrushProperties->BrushSize.DecreaseRadius(true);
+	NotifyOfPropertyChangeByTool(BrushProperties);
 	CalculateBrushRadius();
 }
 
diff --git a/Engine/Source/Runtime/InteractiveToolsFramework/Private/BaseTools/BaseBrushTool.cpp b/Engine/Source/Runtime/InteractiveToolsFramework/Private/BaseTools/BaseBrushTool.cpp
--- a/Engine/Source/Runtime/InteractiveToolsFramework/Private/BaseTools/BaseBrushTool.cpp
+++ b/Engine/Source/Runtime/InteractiveToolsFramework/Private/BaseTools/BaseBrushTool.cpp
@@ -63,12 +63,14 @@ void UBaseBrushTool::IncreaseBrushSizeAction()
 {
 	BrushProperties->BrushSize = FMath::Clamp(BrushProperties->BrushSize + 0.025f, 0.0f, 1.0f);
 	RecalculateBrushRadius();
+	NotifyOfPropertyChangeByTool(BrushProperties);
 }
 
 void UBaseBrushTool::DecreaseBrushSizeAction()
 {
 	BrushProperties->BrushSize = FMath::Clamp(BrushProperties->BrushSize - 0.025f, 0.0f, 1.0f);
 	RecalculateBrushRadius();
+	NotifyOfPropertyChangeByTool(BrushProperties);
 }
 
 
@@ -79,6 +81,7 @@ void UBaseBrushTool::IncreaseBrushStrengthAction()
 
 	float NewValue = OldValue + ChangeAmount;
 	BrushProperties->BrushStrength = FMath::Clamp(NewValue, 0.f, 1.f);
+	NotifyOfPropertyChangeByTool(BrushProperties);
 }
 
 void UBaseBrushTool::DecreaseBrushStrengthAction()
@@ -88,6 +91,7 @@ void UBaseBrushTool::DecreaseBrushStrengthAction()
 
 	float NewValue = OldValue - ChangeAmount;
 	BrushProperties->BrushStrength = FMath::Clamp(NewValue, 0.f, 1.f);
+	NotifyOfPropertyChangeByTool(BrushProperties);
 }
 
 
@@ -98,6 +102,7 @@ void UBaseBrushTool::IncreaseBrushFalloffAction()
 
 	float NewValue = OldValue + ChangeAmount;
 	BrushProperties->BrushFalloffAmount = FMath::Clamp(NewValue, 0.f, 1.f);
+	NotifyOfPropertyChangeByTool(BrushProperties);
 }
 
 void UBaseBrushTool::DecreaseBrushFalloffAction()
@@ -107,6 +112,7 @@ void UBaseBrushTool::DecreaseBrushFalloffAction()
 
 	float NewValue = OldValue - ChangeAmount;
 	BrushProperties->BrushFalloffAmount = FMath::Clamp(NewValue, 0.f, 1.f);
+	NotifyOfPropertyChangeByTool(BrushProperties);
 }
 
 void UBaseBrushTool::RegisterActions(FInteractiveToolActionSet& ActionSet)
