ModelingTools: Fix issue with PolyExt Tool failing to register correct initial point placement when snapping modes are disabled. Increased clarity around snapping parameter comments.

#rnx

diff --git a/Engine/Source/Runtime/InteractiveToolsFramework/Public/SceneQueries/SceneSnappingManager.h b/Engine/Source/Runtime/InteractiveToolsFramework/Public/SceneQueries/SceneSnappingManager.h
--- a/Engine/Source/Runtime/InteractiveToolsFramework/Public/SceneQueries/SceneSnappingManager.h
+++ b/Engine/Source/Runtime/InteractiveToolsFramework/Public/SceneQueries/SceneSnappingManager.h
@@ -116,12 +116,16 @@ struct INTERACTIVETOOLSFRAMEWORK_API FSceneSnapQueryRequest
 	TOptional<FRotator> RotGridSize{};
 
 	/** Snap input position */
-	FVector Position;
-	/** Another position must deviate less than this number of degrees (in visual angle) to be considered an acceptable snap position */
-	float VisualAngleThresholdDegrees;
+	FVector Position = FVector::ZeroVector;
+
+	/**
+	 *  When considering if one point is close enough to another point for snapping purposes, they
+	 *  must deviate less than this number of degrees (in visual angle) to be considered an acceptable snap position.
+	 */
+	float VisualAngleThresholdDegrees = 15.0;
 
 	/** Snap input rotation delta */
-	FQuat DeltaRotation;
+	FQuat DeltaRotation = FQuat(EForceInit::ForceInitToZero);
 
 	/** Optional: components to consider invisible even if they aren't. */
 	const TArray<const UPrimitiveComponent*>* ComponentsToIgnore = nullptr;
