InteractiveToolsFramework: Add tooltips to EToolContextCoordinateSystem values

diff --git a/Engine/Source/Runtime/InteractiveToolsFramework/Public/ToolContextInterfaces.h b/Engine/Source/Runtime/InteractiveToolsFramework/Public/ToolContextInterfaces.h
--- a/Engine/Source/Runtime/InteractiveToolsFramework/Public/ToolContextInterfaces.h
+++ b/Engine/Source/Runtime/InteractiveToolsFramework/Public/ToolContextInterfaces.h
@@ -99,7 +99,9 @@ enum class EStandardToolContextMaterials
 UENUM()
 enum class EToolContextCoordinateSystem
 {
+	/** World space coordinate system */
 	World = 0,
+	/** Local coordinate system */
 	Local = 1
 };
 
