Minor dependency cleanup in ComponentSourceInterfaces.h

diff --git a/Engine/Plugins/Experimental/MeshModelingToolsetExp/Source/MeshModelingToolsExp/Private/MeshVertexSculptTool.cpp b/Engine/Plugins/Experimental/MeshModelingToolsetExp/Source/MeshModelingToolsExp/Private/MeshVertexSculptTool.cpp
--- a/Engine/Plugins/Experimental/MeshModelingToolsetExp/Source/MeshModelingToolsExp/Private/MeshVertexSculptTool.cpp
+++ b/Engine/Plugins/Experimental/MeshModelingToolsetExp/Source/MeshModelingToolsExp/Private/MeshVertexSculptTool.cpp
@@ -31,6 +31,7 @@
 #include "Sculpting/MeshSculptUtil.h"
 
 #include "TargetInterfaces/PrimitiveComponentBackedTarget.h"
+#include "TargetInterfaces/MaterialProvider.h"
 
 using namespace UE::Geometry;
 
diff --git a/Engine/Plugins/Runtime/MeshModelingToolset/Source/ModelingComponents/Public/Drawing/UVLayoutPreview.h b/Engine/Plugins/Runtime/MeshModelingToolset/Source/ModelingComponents/Public/Drawing/UVLayoutPreview.h
--- a/Engine/Plugins/Runtime/MeshModelingToolset/Source/ModelingComponents/Public/Drawing/UVLayoutPreview.h
+++ b/Engine/Plugins/Runtime/MeshModelingToolset/Source/ModelingComponents/Public/Drawing/UVLayoutPreview.h
@@ -6,6 +6,7 @@
 #include "PreviewMesh.h"
 #include "TriangleSetComponent.h"
 #include "InteractiveTool.h"
+#include "TargetInterfaces/MaterialProvider.h"
 #include "UVLayoutPreview.generated.h"
 
 PREDECLARE_USE_GEOMETRY_CLASS(FDynamicMesh3);
diff --git a/Engine/Plugins/Runtime/MeshModelingToolset/Source/ModelingComponentsEditorOnly/Private/ToolTargets/VolumeComponentTarget.cpp b/Engine/Plugins/Runtime/MeshModelingToolset/Source/ModelingComponentsEditorOnly/Private/ToolTargets/VolumeComponentTarget.cpp
--- a/Engine/Plugins/Runtime/MeshModelingToolset/Source/ModelingComponentsEditorOnly/Private/ToolTargets/VolumeComponentTarget.cpp
+++ b/Engine/Plugins/Runtime/MeshModelingToolset/Source/ModelingComponentsEditorOnly/Private/ToolTargets/VolumeComponentTarget.cpp
@@ -11,6 +11,7 @@
 #include "DynamicMesh/MeshNormals.h"
 #include "StaticMeshAttributes.h"
 #include "ToolSetupUtil.h"
+#include "TargetInterfaces/MaterialProvider.h"
 
 using namespace UE::Geometry;
 
diff --git a/Engine/Source/Editor/UnrealEd/Private/Tools/EditorComponentSourceFactory.cpp b/Engine/Source/Editor/UnrealEd/Private/Tools/EditorComponentSourceFactory.cpp
--- a/Engine/Source/Editor/UnrealEd/Private/Tools/EditorComponentSourceFactory.cpp
+++ b/Engine/Source/Editor/UnrealEd/Private/Tools/EditorComponentSourceFactory.cpp
@@ -10,6 +10,7 @@
 #include "PhysicsEngine/BodySetup.h"
 #include "Widgets/Notifications/SNotificationList.h"
 #include "Framework/Notifications/NotificationManager.h"
+#include "TargetInterfaces/MaterialProvider.h"
 
 static void DisplayCriticalWarningMessage(const FString& Message)
 {
diff --git a/Engine/Source/Runtime/GeometryFramework/Private/Components/MeshRenderDecomposition.cpp b/Engine/Source/Runtime/GeometryFramework/Private/Components/MeshRenderDecomposition.cpp
--- a/Engine/Source/Runtime/GeometryFramework/Private/Components/MeshRenderDecomposition.cpp
+++ b/Engine/Source/Runtime/GeometryFramework/Private/Components/MeshRenderDecomposition.cpp
@@ -6,6 +6,7 @@
 #include "DynamicMesh/DynamicMeshAABBTree3.h"
 #include "Async/ParallelFor.h"
 #include "ComponentSourceInterfaces.h"
+#include "TargetInterfaces/MaterialProvider.h"
 
 using namespace UE::Geometry;
 
diff --git a/Engine/Source/Runtime/InteractiveToolsFramework/Private/ComponentSourceInterfaces.cpp b/Engine/Source/Runtime/InteractiveToolsFramework/Private/ComponentSourceInterfaces.cpp
--- a/Engine/Source/Runtime/InteractiveToolsFramework/Private/ComponentSourceInterfaces.cpp
+++ b/Engine/Source/Runtime/InteractiveToolsFramework/Private/ComponentSourceInterfaces.cpp
@@ -4,6 +4,7 @@
 
 #include "Components/PrimitiveComponent.h"
 #include "Containers/Array.h"
+#include "TargetInterfaces/MaterialProvider.h"
 
 
 namespace
diff --git a/Engine/Source/Runtime/InteractiveToolsFramework/Public/ComponentSourceInterfaces.h b/Engine/Source/Runtime/InteractiveToolsFramework/Public/ComponentSourceInterfaces.h
--- a/Engine/Source/Runtime/InteractiveToolsFramework/Public/ComponentSourceInterfaces.h
+++ b/Engine/Source/Runtime/InteractiveToolsFramework/Public/ComponentSourceInterfaces.h
@@ -6,8 +6,6 @@
 #include "Math/UnrealMath.h"
 #include "Misc/Optional.h"
 #include "Templates/Function.h"
-#include "TargetInterfaces/MaterialProvider.h" // FComponentMaterialSet
-#include "Engine/EngineTypes.h"    // FHitResult
 
 // predeclarations
 class AActor;
@@ -15,7 +13,8 @@ class UActorComponent;
 class UPrimitiveComponent;
 struct FMeshDescription;
 class UMaterialInterface;
-
+struct FHitResult;
+struct FComponentMaterialSet;
 
 /**
  * @deprecated Use tool targets instead, such as UPrimitiveComponentToolTarget
