CIS fix - missing abstract class specifier.

diff --git a/Engine/Source/Runtime/InteractiveToolsFramework/Public/BaseGizmos/GizmoElementBase.h b/Engine/Source/Runtime/InteractiveToolsFramework/Public/BaseGizmos/GizmoElementBase.h
--- a/Engine/Source/Runtime/InteractiveToolsFramework/Public/BaseGizmos/GizmoElementBase.h
+++ b/Engine/Source/Runtime/InteractiveToolsFramework/Public/BaseGizmos/GizmoElementBase.h
@@ -50,7 +50,7 @@ enum class EGizmoElementViewAlignType
  * Contains common properties and utility functions.
  * This class does nothing by itself, use subclasses like UGizmoElementCylinder
  */
-UCLASS(Transient)
+UCLASS(Transient, Abstract)
 class INTERACTIVETOOLSFRAMEWORK_API UGizmoElementBase : public UObject
 {
 	GENERATED_BODY()
