Ray-cone and ray-cylinder intersection fixes: clipping cone at height and zero-comparison precision issues.

diff --git a/Engine/Source/Runtime/InteractiveToolsFramework/Private/BaseGizmos/GizmoMath.cpp b/Engine/Source/Runtime/InteractiveToolsFramework/Private/BaseGizmos/GizmoMath.cpp
--- a/Engine/Source/Runtime/InteractiveToolsFramework/Private/BaseGizmos/GizmoMath.cpp
+++ b/Engine/Source/Runtime/InteractiveToolsFramework/Private/BaseGizmos/GizmoMath.cpp
@@ -171,7 +171,7 @@
 	// Get the z-value, in cylinder coordinates, of the incoming line's
 	// unit-length direction.
 	RealType Dz = Basis[0] | RayDirection;
-	if (FMath::Abs(Dz) == 1.0)
+	if (FMath::IsNearlyEqual(Dz, 1.0, SMALL_NUMBER)) 
 	{
 		// The line is parallel to the cylinder axis.  Determine whether the
 		// line intersects the cylinder end disks.
@@ -201,7 +201,7 @@
 
 		RealType A0, A1, A2, Discr, Root, Inv, TValue;
 
-		if (D[2] == 0.0)
+		if (FMath::IsNearlyZero(D[2]))
 		{
 			// The line is perpendicular to the cylinder axis.
 			if (FMath::Abs(P[2]) <= HalfHeight)
@@ -215,7 +215,14 @@
 				A1 = P[0] * D[0] + P[1] * D[1];
 				A2 = D[0] * D[0] + D[1] * D[1];
 				Discr = A1 * A1 - A0 * A2;
-				if (Discr > 0.0)
+				if (FMath::IsNearlyZero(Discr)) 
 				{
+					// The line is tangent to the cylinder.
+					NumIntersections = 1;
+					RayParam[0] = -A1 / A2;
+					RayParam[1] = RayParam[0];
+				}
+				else if (Discr > 0.0)
+				{
 					// The line intersects the cylinder in two places.
 					NumIntersections = 2;
@@ -223,14 +230,7 @@
 					Inv = 1.0f / A2;
 					RayParam[0] = (-A1 - Root) * Inv;
 					RayParam[1] = (-A1 + Root) * Inv;
-				}
-				else if (Discr == 0.0)
-				{
-					// The line is tangent to the cylinder.
-					NumIntersections = 1;
-					RayParam[0] = -A1 / A2;
-					RayParam[1] = RayParam[0];
-				}
+				}				
 				// else: The line does not intersect the cylinder.
 			}
 			// else: The line is outside the planes of the cylinder end disks.
@@ -265,7 +265,25 @@
 				A1 = P[0] * D[0] + P[1] * D[1];
 				A2 = D[0] * D[0] + D[1] * D[1];
 				Discr = A1 * A1 - A0 * A2;
-				if (Discr > 0.0f)
+				if (FMath::IsNearlyZero(Discr))
 				{
+					TValue = -A1 / A2;
+					if (T0 <= T1)
+					{
+						if (T0 <= TValue && TValue <= T1)
+						{
+							RayParam[NumIntersections++] = TValue;
+						}
+					}
+					else
+					{
+						if (T1 <= TValue && TValue <= T0)
+						{
+							RayParam[NumIntersections++] = TValue;
+						}
+					}
+				}
+				else if (Discr > 0.0)
+				{
 					Root = FMath::Sqrt(Discr);
 					Inv = (1.0f) / A2;
@@ -305,24 +323,6 @@
 					}
 					// else: Line intersects end disk and cylinder wall.
 				}
-				else if (Discr == 0.0f)
-				{
-					TValue = -A1 / A2;
-					if (T0 <= T1)
-					{
-						if (T0 <= TValue && TValue <= T1)
-						{
-							RayParam[NumIntersections++] = TValue;
-						}
-					}
-					else
-					{
-						if (T1 <= TValue && TValue <= T0)
-						{
-							RayParam[NumIntersections++] = TValue;
-						}
-					}
-				}
 				// else: Line does not intersect cylinder wall.
 			}
 			// else: Line intersects both top and bottom cylinder end disks.
@@ -405,17 +405,34 @@
 	RealType T;
 	RealType RayParam[2];
 
-	if (C2 != 0.0)
+	if (!FMath::IsNearlyZero(C2))
 	{
 		RealType Discr = C1 * C1 - C0 * C2;
-		if (Discr < 0.0)
+		if (FMath::IsNearlyZero(Discr))
 		{
+			// One repeated real Root; the line is tangent to the double-sided
+			// cone at a single point.  Report only the point if it is on the
+			// positive cone.
+			T = -C1 / C2;
+			if (DdU * T + DdPmV >= 0.0)
+			{
+				RayParam[0] = T;
+				RayParam[1] = T;
+			}
+			else
+			{
+				bIntersectsOut = false;
+				return;
+			}
+		}
+		else if (Discr < 0.0)
+		{
 			// The quadratic has no real-valued Roots.  The line does not
 			// intersect the double-sided cone.
 			bIntersectsOut = false;
 			return;
 		}
-		else if (Discr > 0.0)
+		else // (Discr > 0.0)
 		{
 			// The quadratic has two distinct real-valued Roots.  However, one
 			// or both of them might intersect the negative cone.  We are
@@ -468,23 +485,6 @@
 				return;
 			}
 		}
-		else  // Discr == 0
-		{
-			// One repeated real Root; the line is tangent to the double-sided
-			// cone at a single point.  Report only the point if it is on the
-			// positive cone.
-			T = -C1 / C2;
-			if (DdU * T + DdPmV >= 0.0)
-			{
-				RayParam[0] = T;
-				RayParam[1] = T;
-			}
-			else
-			{
-				bIntersectsOut = false;
-				return;
-			}
-		}
 	}
 	else
 	{
@@ -492,7 +492,7 @@
 		return;
 	}
 
-	if (DdU != 0.0)
+	if (!FMath::IsNearlyZero(DdU))
 	{
 		// Clamp the intersection to the height of the cone.
 		RealType InvDdU = (1.0f) / DdU;
@@ -502,7 +502,7 @@
 			hInterval[0] = -DdPmV * InvDdU;
 			hInterval[1] = (ConeHeight - DdPmV) * InvDdU;
 		}
-		else // (DdU < 0.0f)
+		else // (DdU < 0.0)
 		{
 			hInterval[0] = (ConeHeight - DdPmV) * InvDdU;
 			hInterval[1] = -DdPmV * InvDdU;
@@ -523,13 +523,10 @@
 			return;
 		}
 	}
-	else if (bIntersectsOut)
-	{
-		if (DdPmV > ConeHeight)
-		{
-			bIntersectsOut = false;
-			return;
-		}
+	else if (DdPmV > ConeHeight)
+	{		
+		bIntersectsOut = false;
+		return;
 	}
 
 	// Get rid of hits before ray origin
