Gizmo element library: added linetrace traversal support, implemented cone/cylinder/arrow hit-testing.

diff --git a/Engine/Source/Runtime/InteractiveToolsFramework/Private/BaseGizmos/GizmoElementArrow.cpp b/Engine/Source/Runtime/InteractiveToolsFramework/Private/BaseGizmos/GizmoElementArrow.cpp
--- a/Engine/Source/Runtime/InteractiveToolsFramework/Private/BaseGizmos/GizmoElementArrow.cpp
+++ b/Engine/Source/Runtime/InteractiveToolsFramework/Private/BaseGizmos/GizmoElementArrow.cpp
@@ -86,6 +86,11 @@ FInputRayHit UGizmoElementArrow::LineTrace(const FVector RayOrigin, const FVecto
 		}
 	}
 
+	if (Hit.bHit)
+	{
+		Hit.HitOwner = this;
+	}
+
 	return Hit;
 }
 
@@ -229,11 +234,13 @@ void UGizmoElementArrow::UpdateArrowHead()
 	if (HeadType == EGizmoElementArrowHeadType::Cone)
 	{
 		check(ConeElement);
-		ConeElement->SetBase(Direction * (BodyLength - HeadLength * 0.1f));
-		ConeElement->SetDirection(Direction);
+		// head length is multiplied by 0.9f to prevent gap between body cylinder and head cone
+		ConeElement->SetOrigin(Direction * (BodyLength + HeadLength * 0.9f)); 
+		ConeElement->SetDirection(-Direction);
 		ConeElement->SetHeight(HeadLength);
 		ConeElement->SetRadius(HeadRadius);
 		ConeElement->SetNumSides(NumSides);
+		ConeElement->SetElementInteractionState(ElementInteractionState);
 	}
 	else // (HeadType == EGizmoElementArrowHeadType::Cube)
 	{
@@ -242,5 +249,6 @@ void UGizmoElementArrow::UpdateArrowHead()
 		BoxElement->SetUpDirection(Direction);
 		BoxElement->SetSideDirection(SideDirection);
 		BoxElement->SetDimensions(FVector(HeadLength, HeadLength, HeadLength));
+		BoxElement->SetElementInteractionState(ElementInteractionState);
 	}
 }
diff --git a/Engine/Source/Runtime/InteractiveToolsFramework/Private/BaseGizmos/GizmoElementBase.cpp b/Engine/Source/Runtime/InteractiveToolsFramework/Private/BaseGizmos/GizmoElementBase.cpp
--- a/Engine/Source/Runtime/InteractiveToolsFramework/Private/BaseGizmos/GizmoElementBase.cpp
+++ b/Engine/Source/Runtime/InteractiveToolsFramework/Private/BaseGizmos/GizmoElementBase.cpp
@@ -56,7 +56,7 @@ bool UGizmoElementBase::GetViewAlignRot(const FSceneView* View, const FTransform
 		bool bNonUniformScaleWarning = true;
 		if (bNonUniformScaleWarning)
 		{
-			UE_LOG(LogGizmoElements, Warning, TEXT("Gizmo element library view-dependent alignment does not currently support non-uniform scale (%f, %f %f)."),
+			UE_LOG(LogGizmoElements, Warning, TEXT("Gizmo element library view-dependent alignment does not currently support non-uniform scale (%f %f %f)."),
 				Scale.X, Scale.Y, Scale.Z);
 			bNonUniformScaleWarning = false;
 		}
@@ -130,44 +130,43 @@ bool UGizmoElementBase::GetViewAlignRot(const FSceneView* View, const FTransform
 	return true;
 }
 
-const UMaterialInterface* UGizmoElementBase::GetCurrentMaterial(const FRenderTraversalState& RenderState, EGizmoElementInteractionState InState) const
+const UMaterialInterface* UGizmoElementBase::GetCurrentMaterial(const FRenderTraversalState& RenderState) const
 {
-	if (InState == EGizmoElementInteractionState::Hovering)
+	EGizmoElementInteractionState CurrentState;
+
+	if (RenderState.InteractionState == EGizmoElementInteractionState::None)
 	{
-		if (HoverMaterial)
-		{
-			return HoverMaterial;
-		}
-		else if (RenderState.HoverMaterial.IsValid())
+		CurrentState = ElementInteractionState;
+	}
+	else
+	{
+		CurrentState = RenderState.InteractionState;
+	}
+
+	if (CurrentState == EGizmoElementInteractionState::Hovering)
+	{
+		if (RenderState.HoverMaterial.IsValid())
 		{
 			return RenderState.HoverMaterial.Get();
 		}
+		return HoverMaterial;
 	}
-	else if (InState == EGizmoElementInteractionState::Interacting)
+
+	if (CurrentState == EGizmoElementInteractionState::Interacting)
 	{
-		if (InteractMaterial)
-		{
-			return InteractMaterial;
-		}
-		else if (RenderState.InteractMaterial.IsValid())
+		if (RenderState.InteractMaterial.IsValid())
 		{
 			return RenderState.InteractMaterial.Get();
 		}
+		return InteractMaterial;
 	}
-	else
-	{
-		if (Material)
-		{
-			return Material;
-		}
-		else if (RenderState.Material.IsValid())
-		{
-			return RenderState.Material.Get();
-		}
 
+	// CurrentState is None, so just return the regular material
+	if (RenderState.Material.IsValid())
+	{
+		return RenderState.Material.Get();
 	}
-
-	return nullptr;
+	return Material;
 }
 
 void UGizmoElementBase::CacheRenderState(const FTransform& InLocalToWorldState, bool InVisibleViewDependent)
@@ -185,15 +184,19 @@ void UGizmoElementBase::ResetCachedRenderState()
 
 void UGizmoElementBase::UpdateRenderTraversalState(FRenderTraversalState& InRenderTraversalState)
 {
-	if (HoverMaterial)
+	if (InRenderTraversalState.InteractionState == EGizmoElementInteractionState::None)
+	{
+		InRenderTraversalState.InteractionState = ElementInteractionState;
+	}
+	if (!InRenderTraversalState.HoverMaterial.IsValid())
 	{
 		InRenderTraversalState.HoverMaterial = HoverMaterial;
 	}
-	if (InteractMaterial)
+	if (!InRenderTraversalState.InteractMaterial.IsValid())
 	{
 		InRenderTraversalState.InteractMaterial = InteractMaterial;
 	}
-	if (Material)
+	if (!InRenderTraversalState.Material.IsValid())
 	{
 		InRenderTraversalState.Material = Material;
 	}
diff --git a/Engine/Source/Runtime/InteractiveToolsFramework/Private/BaseGizmos/GizmoElementBox.cpp b/Engine/Source/Runtime/InteractiveToolsFramework/Private/BaseGizmos/GizmoElementBox.cpp
--- a/Engine/Source/Runtime/InteractiveToolsFramework/Private/BaseGizmos/GizmoElementBox.cpp
+++ b/Engine/Source/Runtime/InteractiveToolsFramework/Private/BaseGizmos/GizmoElementBox.cpp
@@ -23,7 +23,7 @@ void UGizmoElementBox::Render(IToolsContextRenderAPI* RenderAPI, const FRenderTr
 
 	if (bVisibleViewDependent)
 	{
-		const UMaterialInterface* UseMaterial = GetCurrentMaterial(RenderState, ElementInteractionState);
+		const UMaterialInterface* UseMaterial = GetCurrentMaterial(RenderState);
 
 		if (UseMaterial)
 		{
diff --git a/Engine/Source/Runtime/InteractiveToolsFramework/Private/BaseGizmos/GizmoElementCircle.cpp b/Engine/Source/Runtime/InteractiveToolsFramework/Private/BaseGizmos/GizmoElementCircle.cpp
--- a/Engine/Source/Runtime/InteractiveToolsFramework/Private/BaseGizmos/GizmoElementCircle.cpp
+++ b/Engine/Source/Runtime/InteractiveToolsFramework/Private/BaseGizmos/GizmoElementCircle.cpp
@@ -55,7 +55,7 @@ void UGizmoElementCircle::Render(IToolsContextRenderAPI* RenderAPI, const FRende
 
 		if (bDrawMesh)
 		{
-			if (const UMaterialInterface* UseMaterial = GetCurrentMaterial(RenderState, ElementInteractionState))
+			if (const UMaterialInterface* UseMaterial = GetCurrentMaterial(RenderState))
 			{
 				DrawDisc(PDI, WorldCenter, Axis0, Axis1, VertexColor, WorldRadius, NumSides, UseMaterial->GetRenderProxy(), SDPG_Foreground);
 			}
diff --git a/Engine/Source/Runtime/InteractiveToolsFramework/Private/BaseGizmos/GizmoElementCone.cpp b/Engine/Source/Runtime/InteractiveToolsFramework/Private/BaseGizmos/GizmoElementCone.cpp
--- a/Engine/Source/Runtime/InteractiveToolsFramework/Private/BaseGizmos/GizmoElementCone.cpp
+++ b/Engine/Source/Runtime/InteractiveToolsFramework/Private/BaseGizmos/GizmoElementCone.cpp
@@ -19,28 +19,27 @@ void UGizmoElementCone::Render(IToolsContextRenderAPI* RenderAPI, const FRenderT
 
 	FTransform LocalToWorldTransform = RenderState.LocalToWorldTransform;
 
-	bool bVisibleViewDependent = GetViewDependentVisibility(View, LocalToWorldTransform, Base);
+	bool bVisibleViewDependent = GetViewDependentVisibility(View, LocalToWorldTransform, Origin);
 
 	if (bVisibleViewDependent)
 	{
-		const UMaterialInterface* UseMaterial = GetCurrentMaterial(RenderState, ElementInteractionState);
+		const UMaterialInterface* UseMaterial = GetCurrentMaterial(RenderState);
 
 		if (UseMaterial)
 		{
 			FQuat AlignRot;
-			const bool bHasAlignRot = GetViewAlignRot(View, LocalToWorldTransform, Base, AlignRot);
+			const bool bHasAlignRot = GetViewAlignRot(View, LocalToWorldTransform, Origin, AlignRot);
 			const FVector AdjustedDir = bHasAlignRot ? AlignRot.RotateVector(Direction) : Direction;
 
-			FQuat Rotation = FRotationMatrix::MakeFromX(-AdjustedDir).ToQuat();
+			FQuat Rotation = FRotationMatrix::MakeFromX(AdjustedDir).ToQuat();
 
 			const FVector Scale(Height);
-			const FVector Origin = Base + AdjustedDir * Height;
 
-			LocalToWorldTransform = FTransform(Rotation, Origin, Scale) * LocalToWorldTransform;
+			FTransform RenderLocalToWorldTransform = FTransform(Rotation, Origin, Scale) * LocalToWorldTransform;
 			const float ConeSide = FMath::Sqrt(Height * Height + Radius * Radius);
 			const float Angle = FMath::Acos(Height / ConeSide);
 			FPrimitiveDrawInterface* PDI = RenderAPI->GetPrimitiveDrawInterface();
-			DrawCone(PDI, LocalToWorldTransform.ToMatrixWithScale(), Angle, Angle, NumSides, false, FColor::White, UseMaterial->GetRenderProxy(), SDPG_Foreground);
+			DrawCone(PDI, RenderLocalToWorldTransform.ToMatrixWithScale(), Angle, Angle, NumSides, false, FColor::White, UseMaterial->GetRenderProxy(), SDPG_Foreground);
 		}
 	}
 	CacheRenderState(LocalToWorldTransform, bVisibleViewDependent);
@@ -52,33 +51,29 @@ FInputRayHit UGizmoElementCone::LineTrace(const FVector RayOrigin, const FVector
 {
 	if (IsHittableInView())
 	{
-		// @todo - modify ray-cone intersection to work with updated properties
-
-#if 0
 		bool bIntersects = false;
 		double RayParam = 0.0;
 
-		const FMatrix LocalToWorldMatrix = LocalToWorldTransform.ToMatrixNoScale();
-
-		const FVector Origin = LocalToWorldMatrix.TransformPosition(FVector::ZeroVector);
-		const FVector ConeDirection = (bWorld) ? Direction : FVector{ LocalToWorldMatrix.TransformVector(Direction) };
-		const double ConeHeight = Height * DynamicPixelToWorldScale;
-		const double ConeOffset = Offset * DynamicPixelToWorldScale;
-		const FVector ConeOrigin = Origin + ConeOffset * ConeDirection;
+		const double ConeSide = FMath::Sqrt(Height * Height + Radius * Radius);
+		const double CosAngle = Height / ConeSide;
+		const double ConeHeight = Height * CachedLocalToWorldTransform.GetScale3D().X;
+		const FVector ConeOrigin = CachedLocalToWorldTransform.TransformPosition(Origin);
+		const FVector ConeDirection = CachedLocalToWorldTransform.TransformVectorNoScale(Direction);
 
 		GizmoMath::RayConeIntersection(
-			ConeOrigin, 
-			ConeDirection, 
-			FMath::Cos(Angle), 
+			ConeOrigin,
+			ConeDirection,
+			CosAngle,
 			ConeHeight,
 			RayOrigin, RayDirection,
 			bIntersects, RayParam);
 
 		if (bIntersects)
 		{
-			return FInputRayHit(RayParam);
+			FInputRayHit RayHit(RayParam);
+			RayHit.HitOwner = this;
+			return RayHit;
 		}
-#endif
 	}
 
 	return FInputRayHit();
@@ -90,14 +85,14 @@ FBoxSphereBounds UGizmoElementCone::CalcBounds(const FTransform& LocalToWorld) c
 	return FBoxSphereBounds();
 }
 
-void UGizmoElementCone::SetBase(const FVector& InBase)
+void UGizmoElementCone::SetOrigin(const FVector& InOrigin)
 {
-	Base = InBase;
+	Origin = InOrigin;
 }
 
-FVector UGizmoElementCone::GetBase() const
+FVector UGizmoElementCone::GetOrigin() const
 {
-	return Base;
+	return Origin;
 }
 
 void UGizmoElementCone::SetDirection(const FVector& InDirection)
diff --git a/Engine/Source/Runtime/InteractiveToolsFramework/Private/BaseGizmos/GizmoElementCylinder.cpp b/Engine/Source/Runtime/InteractiveToolsFramework/Private/BaseGizmos/GizmoElementCylinder.cpp
--- a/Engine/Source/Runtime/InteractiveToolsFramework/Private/BaseGizmos/GizmoElementCylinder.cpp
+++ b/Engine/Source/Runtime/InteractiveToolsFramework/Private/BaseGizmos/GizmoElementCylinder.cpp
@@ -23,7 +23,7 @@ void UGizmoElementCylinder::Render(IToolsContextRenderAPI* RenderAPI, const FRen
 
 	if (bVisibleViewDependent)
 	{
-		const UMaterialInterface* UseMaterial = GetCurrentMaterial(RenderState, ElementInteractionState);
+		const UMaterialInterface* UseMaterial = GetCurrentMaterial(RenderState);
 
 		if (UseMaterial)
 		{
@@ -32,12 +32,12 @@ void UGizmoElementCylinder::Render(IToolsContextRenderAPI* RenderAPI, const FRen
 			const FVector AdjustedDir = bHasAlignRot ? AlignRot.RotateVector(Direction) : Direction;
 
 			const FQuat Rotation = FRotationMatrix::MakeFromZ(AdjustedDir).ToQuat();
-			const float HalfHeight = Height * 0.5f;
+			const double HalfHeight = Height * 0.5;
 			const FVector Origin = Base + AdjustedDir * HalfHeight;
 
-			LocalToWorldTransform = FTransform(Rotation, Origin) * LocalToWorldTransform;
+			FTransform RenderLocalToWorldTransform = FTransform(Rotation, Origin) * LocalToWorldTransform;
 			FPrimitiveDrawInterface* PDI = RenderAPI->GetPrimitiveDrawInterface();
-			DrawCylinder(PDI, LocalToWorldTransform.ToMatrixWithScale(), FVector::ZeroVector, FVector(1, 0, 0), FVector(0, 1, 0), FVector(0, 0, 1), Radius, HalfHeight, NumSides, UseMaterial->GetRenderProxy(), SDPG_Foreground);
+			DrawCylinder(PDI, RenderLocalToWorldTransform.ToMatrixWithScale(), FVector::ZeroVector, FVector(1, 0, 0), FVector(0, 1, 0), FVector(0, 0, 1), Radius, HalfHeight, NumSides, UseMaterial->GetRenderProxy(), SDPG_Foreground);
 
 		}
 	}
@@ -46,38 +46,33 @@ void UGizmoElementCylinder::Render(IToolsContextRenderAPI* RenderAPI, const FRen
 
 
 
-FInputRayHit UGizmoElementCylinder::LineTrace(const FVector RayOrigin, const FVector RayDirectionm)
+FInputRayHit UGizmoElementCylinder::LineTrace(const FVector RayOrigin, const FVector RayDirection)
 {
 	if (IsHittableInView())
 	{
-		// @todo - modify ray-cylinder intersection to work with updated properties
-#if 0
 		bool bIntersects = false;
 		double RayParam = 0.0;
-
-		const FTransform LocalToWorldTransform = FTransform(Base) * InLocalToWorldTransform;
-		const double WorldConeHeight = Height * InLocalToWorldTransform.GetScale3D().X;
-		const FVector ConeOrigin = Base + Height * Direction;
-		const FVector WorldConeOrigin = LocalToWorldTransform.TransformPosition(ConeOrigin);
-		const FVector WorldConeDirection = LocalToWorldTransform.TransformVector(Direction);
-		const double ConeSide = FMath::Sqrt(Height * Height + Radius * Radius);
-		const double CosAngle = Radius / ConeSide;
-
-		GizmoMath::RayConeIntersection(
-			WorldConeOrigin,
-			WorldConeDirection,
-			CosAngle,
-			WorldConeHeight,
+		
+		const double CylinderHeight = Height * CachedLocalToWorldTransform.GetScale3D().X;
+		const double CylinderRadius = Radius * CachedLocalToWorldTransform.GetScale3D().X;
+		const FVector CylinderDirection = CachedLocalToWorldTransform.TransformVectorNoScale(Direction);
+		const FVector CylinderLocalCenter = Base + Direction * Height * 0.5;
+		const FVector CylinderCenter = CachedLocalToWorldTransform.TransformPosition(CylinderLocalCenter);
+
+		GizmoMath::RayCylinderIntersection(
+			CylinderCenter,
+			CylinderDirection,
+			CylinderRadius,
+			CylinderHeight,
 			RayOrigin, RayDirection,
 			bIntersects, RayParam);
 
 		if (bIntersects)
 		{
-			FInputRayHit Hit;
-			Hit.HitObject = this;
-			return Hit;
+			FInputRayHit RayHit(RayParam);
+			RayHit.HitOwner = this;
+			return RayHit;
 		}
-#endif
 	}
 
 	return FInputRayHit();
diff --git a/Engine/Source/Runtime/InteractiveToolsFramework/Private/BaseGizmos/GizmoElementGroup.cpp b/Engine/Source/Runtime/InteractiveToolsFramework/Private/BaseGizmos/GizmoElementGroup.cpp
--- a/Engine/Source/Runtime/InteractiveToolsFramework/Private/BaseGizmos/GizmoElementGroup.cpp
+++ b/Engine/Source/Runtime/InteractiveToolsFramework/Private/BaseGizmos/GizmoElementGroup.cpp
@@ -53,17 +53,27 @@ void UGizmoElementGroup::Render(IToolsContextRenderAPI* RenderAPI, const FRender
 
 FInputRayHit UGizmoElementGroup::LineTrace(const FVector Start, const FVector Direction)
 {
+	FInputRayHit Hit;
+
 	if (IsHittable())
 	{
 		for (UGizmoElementBase* Element : Elements)
 		{
 			if (Element)
 			{
-				return Element->LineTrace(Start, Direction);
+				FInputRayHit NewHit = Element->LineTrace(Start, Direction);
+				if (!Hit.bHit || NewHit.HitDepth < Hit.HitDepth)
+				{
+					Hit = NewHit;
+					if (bHitOwner)
+					{
+						Hit.HitOwner = this;
+					}
+				}
 			}
 		}
 	}
-	return FInputRayHit();
+	return Hit;
 }
 
 FBoxSphereBounds UGizmoElementGroup::CalcBounds(const FTransform& LocalToWorld) const
diff --git a/Engine/Source/Runtime/InteractiveToolsFramework/Private/BaseGizmos/GizmoElementHitTargets.cpp b/Engine/Source/Runtime/InteractiveToolsFramework/Private/BaseGizmos/GizmoElementHitTargets.cpp
--- a/Engine/Source/Runtime/InteractiveToolsFramework/Private/BaseGizmos/GizmoElementHitTargets.cpp
+++ b/Engine/Source/Runtime/InteractiveToolsFramework/Private/BaseGizmos/GizmoElementHitTargets.cpp
@@ -16,18 +16,36 @@ FInputRayHit UGizmoElementHitTarget::IsHit(const FInputDeviceRay& ClickPos) cons
 
 void UGizmoElementHitTarget::UpdateHoverState(bool bHovering)
 {
-	if (GizmoElement)
+	if (!GizmoElement)
+	{
+		return;
+	}
+
+	if (bHovering)
 	{
 		GizmoElement->SetElementInteractionState(EGizmoElementInteractionState::Hovering);
 	}
+	else
+	{
+		GizmoElement->SetElementInteractionState(EGizmoElementInteractionState::None);
+	}
 }
 
 void UGizmoElementHitTarget::UpdateInteractingState(bool bInteracting)
 {
-	if (GizmoElement)
+	if (!GizmoElement)
+	{
+		return;
+	}
+
+	if (bInteracting)
 	{
 		GizmoElement->SetElementInteractionState(EGizmoElementInteractionState::Interacting);
 	}
+	else
+	{
+		GizmoElement->SetElementInteractionState(EGizmoElementInteractionState::None);
+	}
 }
 
 UGizmoElementHitTarget* UGizmoElementHitTarget::Construct(UGizmoElementBase* InGizmoElement, UObject* Outer)
diff --git a/Engine/Source/Runtime/InteractiveToolsFramework/Private/BaseGizmos/GizmoElementRectangle.cpp b/Engine/Source/Runtime/InteractiveToolsFramework/Private/BaseGizmos/GizmoElementRectangle.cpp
--- a/Engine/Source/Runtime/InteractiveToolsFramework/Private/BaseGizmos/GizmoElementRectangle.cpp
+++ b/Engine/Source/Runtime/InteractiveToolsFramework/Private/BaseGizmos/GizmoElementRectangle.cpp
@@ -56,7 +56,7 @@ void UGizmoElementRectangle::Render(IToolsContextRenderAPI* RenderAPI, const FRe
 
 		if (bDrawMesh)
 		{
-			if (const UMaterialInterface* UseMaterial = GetCurrentMaterial(RenderState, ElementInteractionState))
+			if (const UMaterialInterface* UseMaterial = GetCurrentMaterial(RenderState))
 			{
 
 				DrawRectangleMesh(PDI, WorldCenter, Axis0, Axis1, VertexColor, WorldWidth, WorldHeight, UseMaterial->GetRenderProxy(), SDPG_Foreground);
diff --git a/Engine/Source/Runtime/InteractiveToolsFramework/Private/BaseGizmos/GizmoElementTorus.cpp b/Engine/Source/Runtime/InteractiveToolsFramework/Private/BaseGizmos/GizmoElementTorus.cpp
--- a/Engine/Source/Runtime/InteractiveToolsFramework/Private/BaseGizmos/GizmoElementTorus.cpp
+++ b/Engine/Source/Runtime/InteractiveToolsFramework/Private/BaseGizmos/GizmoElementTorus.cpp
@@ -23,7 +23,7 @@ void UGizmoElementTorus::Render(IToolsContextRenderAPI* RenderAPI, const FRender
 
 	if (bVisibleViewDependent)
 	{
-		const UMaterialInterface* UseMaterial = GetCurrentMaterial(RenderState, ElementInteractionState);
+		const UMaterialInterface* UseMaterial = GetCurrentMaterial(RenderState);
 			
 		if (UseMaterial)
 		{
diff --git a/Engine/Source/Runtime/InteractiveToolsFramework/Public/BaseGizmos/GizmoElementBase.h b/Engine/Source/Runtime/InteractiveToolsFramework/Public/BaseGizmos/GizmoElementBase.h
--- a/Engine/Source/Runtime/InteractiveToolsFramework/Public/BaseGizmos/GizmoElementBase.h
+++ b/Engine/Source/Runtime/InteractiveToolsFramework/Public/BaseGizmos/GizmoElementBase.h
@@ -79,14 +79,17 @@ public:
 		// LocalToWorld uniform scale, non-uniform scale is not supported
 		double LocalToWorldScale;
 
-		// Material, element's material should override parent's material
+		// Material, if set, overrides the element's material
 		TWeakObjectPtr<UMaterialInterface> Material = nullptr;
 
-		// Hover material, element's hover material should override parent's hover material
+		// Hover material, if set, overrides the element's hover material
 		TWeakObjectPtr<UMaterialInterface> HoverMaterial = nullptr;
 
-		// Interact material, element's interact material should override parent's interact material
+		// Interact material, if set, overrides the element's interact material
 		TWeakObjectPtr<UMaterialInterface> InteractMaterial = nullptr;
+
+		// Interact state, if not equal to none, overrides the element's interact state
+		EGizmoElementInteractionState InteractionState = EGizmoElementInteractionState::None;
 	};
 
 public:
@@ -107,8 +110,8 @@ public:
 	// on view-dependent alignment settings. So it should be prepended to the local-to-world transform.
 	virtual bool GetViewAlignRot(const FSceneView* View, const FTransform& InLocalToWorldTransform, const FVector& InLocalCenter, FQuat& OutAlignRot) const;
 
-	// Returns the current material for given object state. The materials in the object will override those in the RenderState.
-	virtual const UMaterialInterface* GetCurrentMaterial(const FRenderTraversalState& RenderState, EGizmoElementInteractionState InState) const;
+	// Returns the current material for given object state. The materials in the render state will override those in the element.
+	virtual const UMaterialInterface* GetCurrentMaterial(const FRenderTraversalState& RenderState) const;
 
 	// Reset the cached render state.
 	virtual void ResetCachedRenderState();
diff --git a/Engine/Source/Runtime/InteractiveToolsFramework/Public/BaseGizmos/GizmoElementCone.h b/Engine/Source/Runtime/InteractiveToolsFramework/Public/BaseGizmos/GizmoElementCone.h
--- a/Engine/Source/Runtime/InteractiveToolsFramework/Public/BaseGizmos/GizmoElementCone.h
+++ b/Engine/Source/Runtime/InteractiveToolsFramework/Public/BaseGizmos/GizmoElementCone.h
@@ -26,11 +26,11 @@ public:
 	virtual FBoxSphereBounds CalcBounds(const FTransform& LocalToWorld) const override;
 	//~ End UGizmoElementBase Interface.
 
-	// Location of center of cone's base circle.
-	virtual void SetBase(const FVector& InBase);
-	virtual FVector GetBase() const;
+	// Cone tip location.
+	virtual void SetOrigin(const FVector& InBase);
+	virtual FVector GetOrigin() const;
 
-	// Cone axis direction.
+	// Cone axis direction pointing from tip toward base of cone.
 	virtual void SetDirection(const FVector& InDirection);
 	virtual FVector GetDirection() const;
 
@@ -42,29 +42,29 @@ public:
 	virtual void SetRadius(float InRadius);
 	virtual float GetRadius() const;
 
-	// Number of sides to use when rendering cone.
+	// Number of sides for tessellating cone.
 	virtual void SetNumSides(int32 InNumSides);
 	virtual int32 GetNumSides() const;
 
 protected:
 
-	// Location of center of cone's base circle.
+	// Cone tip location.
 	UPROPERTY()
-	FVector Base = FVector::ZeroVector;
+	FVector Origin = FVector::ZeroVector;
 
-	// Direction of cone's axis.
+	// Cone axis direction pointing from tip toward base of cone.
 	UPROPERTY()
-	FVector Direction = FVector(0.0f, 0.0f, 1.0f);
+	FVector Direction = FVector(0.0f, 0.0f, -1.0f);
 
-	// Height of cone
+	// Cone height.
 	UPROPERTY()
 	float Height = 1.0f;
 
-	// Radius of cone circle.
+	// Cone radius.
 	UPROPERTY()
 	float Radius = 0.5f;
 
-	// Number of sides for tessellating cone
+	// Number of sides for tessellating cone.
 	UPROPERTY()
 	int32 NumSides = 32;
 };
diff --git a/Engine/Source/Runtime/InteractiveToolsFramework/Public/BaseGizmos/GizmoElementCylinder.h b/Engine/Source/Runtime/InteractiveToolsFramework/Public/BaseGizmos/GizmoElementCylinder.h
--- a/Engine/Source/Runtime/InteractiveToolsFramework/Public/BaseGizmos/GizmoElementCylinder.h
+++ b/Engine/Source/Runtime/InteractiveToolsFramework/Public/BaseGizmos/GizmoElementCylinder.h
@@ -42,7 +42,7 @@ public:
 	virtual void SetRadius(float InRadius);
 	virtual float GetRadius() const;
 
-	// Number of sides to use when rendering cylinder.
+	// Number of sides for tessellating cylinder.
 	virtual void SetNumSides(int32 InNumSides);
 	virtual int32 GetNumSides() const;
 
@@ -52,19 +52,19 @@ protected:
 	UPROPERTY()
 	FVector Base = FVector::ZeroVector;
 
-	// Direction of cylinder's axis.
+	// Cylinder axis direction.
 	UPROPERTY()
 	FVector Direction = FVector(0.0f, 0.0f, 1.0f);
 
-	// Height of cylinder
+	// Cylinder height.
 	UPROPERTY()
 	float Height = 1.0f;
 
-	// Radius of cylinder circle.
+	// Cylinder radius.
 	UPROPERTY()
 	float Radius = 0.5f;
 
-	// Number of sides for tessellating cylinder
+	// Number of sides for tessellating cylinder.
 	UPROPERTY()
 	int32 NumSides = 32;
 };
diff --git a/Engine/Source/Runtime/InteractiveToolsFramework/Public/BaseGizmos/GizmoElementGroup.h b/Engine/Source/Runtime/InteractiveToolsFramework/Public/BaseGizmos/GizmoElementGroup.h
--- a/Engine/Source/Runtime/InteractiveToolsFramework/Public/BaseGizmos/GizmoElementGroup.h
+++ b/Engine/Source/Runtime/InteractiveToolsFramework/Public/BaseGizmos/GizmoElementGroup.h
@@ -44,6 +44,11 @@ protected:
 	UPROPERTY()
 	bool bConstantScale = false;
 
+	// When true, this group is treated as a single element such that when LineTrace is called, if any of its sub-elements is hit, 
+	// this group will be returned as the owner of the hit. This should be used when a group of elements should be treated as a single handle.
+	UPROPERTY()
+	bool bHitOwner = false;
+		
 	// Gizmo elements within this group
 	UPROPERTY()
 	TArray<TObjectPtr<UGizmoElementBase>> Elements;
