MeshModelingTools: Remove infinite loop in UGizmoLabdaHitTarget.

diff --git a/Engine/Source/Runtime/InteractiveToolsFramework/Private/BaseGizmos/HitTargets.cpp b/Engine/Source/Runtime/InteractiveToolsFramework/Private/BaseGizmos/HitTargets.cpp
--- a/Engine/Source/Runtime/InteractiveToolsFramework/Private/BaseGizmos/HitTargets.cpp
+++ b/Engine/Source/Runtime/InteractiveToolsFramework/Private/BaseGizmos/HitTargets.cpp
@@ -8,7 +8,7 @@ FInputRayHit UGizmoLambdaHitTarget::IsHit(const FInputDeviceRay& ClickPos) const
 {
 	if (IsHitFunction)
 	{
-		return IsHit(ClickPos);
+		return IsHitFunction(ClickPos);
 	}
 	return FInputRayHit();
 }
