Add IGizmoClickMultiTarget, IGizmoRenderTarget and IGizmoRenderMultiTarget interfaces to facilitate hitting parts within a gizmo element hierarchy. Add HitObject member to FInputRayHit struct to hold UObect-derived owners of hits.

diff --git a/Engine/Source/Runtime/InteractiveToolsFramework/Private/BaseGizmos/GizmoElementArrow.cpp b/Engine/Source/Runtime/InteractiveToolsFramework/Private/BaseGizmos/GizmoElementArrow.cpp
--- a/Engine/Source/Runtime/InteractiveToolsFramework/Private/BaseGizmos/GizmoElementArrow.cpp
+++ b/Engine/Source/Runtime/InteractiveToolsFramework/Private/BaseGizmos/GizmoElementArrow.cpp
@@ -88,7 +88,8 @@ FInputRayHit UGizmoElementArrow::LineTrace(const FVector RayOrigin, const FVecto
 
 	if (Hit.bHit)
 	{
-		Hit.HitOwner = this;
+		Hit.SetHitObject(this);
+		Hit.HitIdentifier = PartIdentifier;
 	}
 
 	return Hit;
diff --git a/Engine/Source/Runtime/InteractiveToolsFramework/Private/BaseGizmos/GizmoElementBase.cpp b/Engine/Source/Runtime/InteractiveToolsFramework/Private/BaseGizmos/GizmoElementBase.cpp
--- a/Engine/Source/Runtime/InteractiveToolsFramework/Private/BaseGizmos/GizmoElementBase.cpp
+++ b/Engine/Source/Runtime/InteractiveToolsFramework/Private/BaseGizmos/GizmoElementBase.cpp
@@ -1,6 +1,7 @@
 // Copyright Epic Games, Inc. All Rights Reserved.
  
 #include "BaseGizmos/GizmoElementBase.h"
+#include "BaseGizmos/GizmoInterfaces.h"
 #include "BaseGizmos/GizmoRenderingUtil.h"
 #include "Materials/MaterialInterface.h"
 
@@ -228,6 +229,15 @@ bool UGizmoElementBase::GetEnabled() const
 {
 	return bEnabled;
 }
+void UGizmoElementBase::SetPartIdentifier(uint32 InPartIdentifier)
+{
+	PartIdentifier = InPartIdentifier;
+}
+
+uint32 UGizmoElementBase::GetPartIdentifier()
+{
+	return PartIdentifier;
+}
 
 void UGizmoElementBase::SetElementState(EGizmoElementState InElementState)
 {
@@ -249,6 +259,36 @@ EGizmoElementInteractionState UGizmoElementBase::GetElementInteractionState() co
 	return ElementInteractionState;
 }
 
+void UGizmoElementBase::UpdatePartHittableState(bool bHittable, uint32 InPartIdentifier)
+{
+	if (InPartIdentifier == PartIdentifier)
+	{
+		uint8 State = static_cast<uint8>(ElementState);
+		uint8 HittableMask = static_cast<uint8>(EGizmoElementState::Hittable);
+		uint8 NewState = (bHittable ? State | HittableMask : State & (~HittableMask));
+		ElementState = static_cast<EGizmoElementState>(NewState);
+	}
+}
+
+void UGizmoElementBase::UpdatePartVisibleState(bool bVisible, uint32 InPartIdentifier)
+{
+	if (InPartIdentifier == PartIdentifier)
+	{
+		uint8 State = static_cast<uint8>(ElementState);
+		uint8 VisibleMask = static_cast<uint8>(EGizmoElementState::Visible);
+		uint8 NewState = (bVisible ? State | VisibleMask : State & (~VisibleMask));
+		ElementState = static_cast<EGizmoElementState>(NewState);
+	}
+}
+
+void UGizmoElementBase::UpdatePartInteractionState(EGizmoElementInteractionState InInteractionState, uint32 InPartIdentifier)
+{
+	if (InPartIdentifier == PartIdentifier)
+	{
+		ElementInteractionState = InInteractionState;
+	}
+}
+
 void UGizmoElementBase::SetViewDependentType(EGizmoElementViewDependentType InViewDependentType)
 {
 	ViewDependentType = InViewDependentType;
diff --git a/Engine/Source/Runtime/InteractiveToolsFramework/Private/BaseGizmos/GizmoElementCone.cpp b/Engine/Source/Runtime/InteractiveToolsFramework/Private/BaseGizmos/GizmoElementCone.cpp
--- a/Engine/Source/Runtime/InteractiveToolsFramework/Private/BaseGizmos/GizmoElementCone.cpp
+++ b/Engine/Source/Runtime/InteractiveToolsFramework/Private/BaseGizmos/GizmoElementCone.cpp
@@ -71,7 +71,8 @@ FInputRayHit UGizmoElementCone::LineTrace(const FVector RayOrigin, const FVector
 		if (bIntersects)
 		{
 			FInputRayHit RayHit(RayParam);
-			RayHit.HitOwner = this;
+			RayHit.SetHitObject(this);
+			RayHit.HitIdentifier = PartIdentifier;
 			return RayHit;
 		}
 	}
diff --git a/Engine/Source/Runtime/InteractiveToolsFramework/Private/BaseGizmos/GizmoElementCylinder.cpp b/Engine/Source/Runtime/InteractiveToolsFramework/Private/BaseGizmos/GizmoElementCylinder.cpp
--- a/Engine/Source/Runtime/InteractiveToolsFramework/Private/BaseGizmos/GizmoElementCylinder.cpp
+++ b/Engine/Source/Runtime/InteractiveToolsFramework/Private/BaseGizmos/GizmoElementCylinder.cpp
@@ -70,7 +70,8 @@ FInputRayHit UGizmoElementCylinder::LineTrace(const FVector RayOrigin, const FVe
 		if (bIntersects)
 		{
 			FInputRayHit RayHit(RayParam);
-			RayHit.HitOwner = this;
+			RayHit.SetHitObject(this);
+			RayHit.HitIdentifier = PartIdentifier;
 			return RayHit;
 		}
 	}
diff --git a/Engine/Source/Runtime/InteractiveToolsFramework/Private/BaseGizmos/GizmoElementGroup.cpp b/Engine/Source/Runtime/InteractiveToolsFramework/Private/BaseGizmos/GizmoElementGroup.cpp
--- a/Engine/Source/Runtime/InteractiveToolsFramework/Private/BaseGizmos/GizmoElementGroup.cpp
+++ b/Engine/Source/Runtime/InteractiveToolsFramework/Private/BaseGizmos/GizmoElementGroup.cpp
@@ -67,7 +67,8 @@ FInputRayHit UGizmoElementGroup::LineTrace(const FVector Start, const FVector Di
 					Hit = NewHit;
 					if (bHitOwner)
 					{
-						Hit.HitOwner = this;
+						Hit.SetHitObject(this);
+						Hit.HitIdentifier = PartIdentifier;
 					}
 				}
 			}
@@ -114,6 +115,45 @@ void UGizmoElementGroup::Remove(UGizmoElementBase* InElement)
 	}
 }
 
+void UGizmoElementGroup::UpdatePartVisibleState(bool bVisible, uint32 InPartIdentifier)
+{
+	Super::UpdatePartVisibleState(bVisible, InPartIdentifier);
+
+	for (UGizmoElementBase* Element : Elements)
+	{
+		if (Element)
+		{
+			Element->UpdatePartVisibleState(bVisible, InPartIdentifier);
+		}
+	}
+}
+
+void UGizmoElementGroup::UpdatePartHittableState(bool bHittable, uint32 InPartIdentifier)
+{
+	Super::UpdatePartHittableState(bHittable, InPartIdentifier);
+
+	for (UGizmoElementBase* Element : Elements)
+	{
+		if (Element)
+		{
+			Element->UpdatePartHittableState(bHittable, InPartIdentifier);
+		}
+	}
+}
+
+void UGizmoElementGroup::UpdatePartInteractionState(EGizmoElementInteractionState InInteractionState, uint32 InPartIdentifier)
+{
+	Super::UpdatePartInteractionState(InInteractionState, InPartIdentifier);
+
+	for (UGizmoElementBase* Element : Elements)
+	{
+		if (Element)
+		{
+			Element->UpdatePartInteractionState(InInteractionState, InPartIdentifier);
+		}
+	}
+}
+
 void UGizmoElementGroup::ResetCachedRenderState()
 {
 	Super::ResetCachedRenderState();
diff --git a/Engine/Source/Runtime/InteractiveToolsFramework/Public/BaseGizmos/GizmoElementBase.h b/Engine/Source/Runtime/InteractiveToolsFramework/Public/BaseGizmos/GizmoElementBase.h
--- a/Engine/Source/Runtime/InteractiveToolsFramework/Public/BaseGizmos/GizmoElementBase.h
+++ b/Engine/Source/Runtime/InteractiveToolsFramework/Public/BaseGizmos/GizmoElementBase.h
@@ -3,6 +3,7 @@
 #pragma once
 
 #include "CoreMinimal.h"
+#include "BaseGizmos/GizmoInterfaces.h"
 #include "InputState.h"
 #include "ToolContextInterfaces.h"
 #include "UObject/GCObject.h"
@@ -63,6 +64,8 @@ public:
 	static constexpr float DefaultViewAlignAngleTol = 0.052f;				// ~3 degrees
 	static constexpr float DefaultViewAlignMaxCosAngleTol = 0.998f;			// Cos(DefaultViewAlignAngleTol)
 
+	static constexpr uint32 DefaultPartIdentifier = 0;						// Default part ID, used for elements that are not associated with any gizmo part
+
 	// Helper struct used to store traversal state during render
 	struct FRenderTraversalState
 	{
@@ -129,6 +132,14 @@ public:
 	virtual void SetEnabled(bool InEnabled);
 	virtual bool GetEnabled() const;
 
+	// For an element hierarchy representing multiple parts of a single gizmo, the part identifier establishes 
+	// a correspondence between a gizmo part and the elements that represent that part. The recognized
+	// part identifier values should be defined in the gizmo. Gizmo part identifiers must be greater than or 
+	// equal to one. Identifier 0 is reserved for the default ID which should be assigned to elements
+	// that do not correspond to any gizmo part, such as non-hittable decorative elements.
+	virtual void SetPartIdentifier(uint32 InPartId);
+	virtual uint32 GetPartIdentifier();
+
 	// Object type bitmask indicating whether this object is visible or hittable or both
 	virtual void SetElementState(EGizmoElementState InElementState);
 	virtual EGizmoElementState GetElementState() const;
@@ -137,6 +148,15 @@ public:
 	virtual void SetElementInteractionState(EGizmoElementInteractionState InInteractionState);
 	virtual EGizmoElementInteractionState GetElementInteractionState() const;
 
+	// Update element's visibility state if element is associated with the specified gizmo part.
+	virtual void UpdatePartVisibleState(bool bVisible, uint32 InPartIdentifier);
+
+	// Update element's hittable state if element is associated with the specified gizmo part.
+	virtual void UpdatePartHittableState(bool bHittable, uint32 InPartIdentifier);
+
+	// Update element's interaction state if element is associated with the specified gizmo part.
+	virtual void UpdatePartInteractionState(EGizmoElementInteractionState InInteractionState, uint32 InPartIdentifier);
+
 	// View-dependent type - None, Axis or Plane. 
 	virtual void SetViewDependentType(EGizmoElementViewDependentType ViewDependentType);
 	virtual EGizmoElementViewDependentType GetViewDependentType() const;
@@ -195,6 +215,10 @@ protected:
 	UPROPERTY()
 	bool bEnabled = true;
 
+	// Part identifier
+	UPROPERTY()
+	uint32 PartIdentifier = DefaultPartIdentifier;
+
 	// Element state - indicates whether object is visible or hittable
 	UPROPERTY()
 	EGizmoElementState ElementState = EGizmoElementState::VisibleAndHittable;
diff --git a/Engine/Source/Runtime/InteractiveToolsFramework/Public/BaseGizmos/GizmoElementGroup.h b/Engine/Source/Runtime/InteractiveToolsFramework/Public/BaseGizmos/GizmoElementGroup.h
--- a/Engine/Source/Runtime/InteractiveToolsFramework/Public/BaseGizmos/GizmoElementGroup.h
+++ b/Engine/Source/Runtime/InteractiveToolsFramework/Public/BaseGizmos/GizmoElementGroup.h
@@ -31,6 +31,15 @@ public:
 	// Remove object from group, if it exists
 	virtual void Remove(UGizmoElementBase* InElement);
 
+	// Update group and contained elements' visibility state for elements in specified gizmo parts.
+	virtual void UpdatePartVisibleState(bool bVisible, uint32 InPartIdentifier);
+
+	// Update group and contained elements' hittable state for elements in specified gizmo parts.
+	virtual void UpdatePartHittableState(bool bHittable, uint32 InPartIdentifier);
+
+	// Update group and contained elements' interaction state for elements in specified gizmo parts.
+	virtual void UpdatePartInteractionState(EGizmoElementInteractionState InInteractionState, uint32 InPartIdentifier);
+
 	// Reset cached render state
 	virtual void ResetCachedRenderState();
 
diff --git a/Engine/Source/Runtime/InteractiveToolsFramework/Public/BaseGizmos/GizmoInterfaces.h b/Engine/Source/Runtime/InteractiveToolsFramework/Public/BaseGizmos/GizmoInterfaces.h
--- a/Engine/Source/Runtime/InteractiveToolsFramework/Public/BaseGizmos/GizmoInterfaces.h
+++ b/Engine/Source/Runtime/InteractiveToolsFramework/Public/BaseGizmos/GizmoInterfaces.h
@@ -119,7 +119,108 @@ public:
 	virtual void UpdateInteractingState(bool bInteracting) = 0;
 };
 
+UINTERFACE()
+class INTERACTIVETOOLSFRAMEWORK_API UGizmoClickMultiTarget : public UInterface
+{
+	GENERATED_BODY()
+};
+/**
+ * IGizmoClickMultiTarget is an interface used to provide a ray-object hit test against a target which
+ * supports hitting parts of the target.
+ *
+ * For a gizmo with multiple parts, the part identifier establishes a correspondence between a gizmo part 
+ * and the elements representing that part within the hit target. The valid part identifiers should 
+ * be defined in the gizmo. Identifier 0 is reserved for the default ID which should be assigned to 
+ * elements that do not correspond to any gizmo part, such as non-hittable decorative elements.
+ */
+class INTERACTIVETOOLSFRAMEWORK_API IGizmoClickMultiTarget
+{
+	GENERATED_BODY()
+public:
+	/**
+	 * @return FInputRayHit indicating whether or not the target object was hit by the device-ray at ClickPos
+	 *         The ray hit contains client-defined ID, HitOwner and HitObject which are used to identify the hit part.
+	 */
+	 //UFUNCTION()    // FInputDeviceRay is not USTRUCT because FRay is not USTRUCT
+	virtual FInputRayHit IsHit(const FInputDeviceRay& ClickPos) const = 0;
+
+	/*
+	 * Updates the hover state of the specified gizmo part, indicating whether the input device is currently hovering 
+	 * over the Standard gizmo.
+	 */
+	UFUNCTION()
+	virtual void UpdateHoverState(bool bHovering, uint32 InPartIdentifier) = 0;
+
+	/*
+	 * Updates the interacting state of the specified gizmo part, indicating when interaction with the 
+	 * Standard gizmo is actively occurring, typically upon the input device clicking and dragging the Standard gizmo. 
+	 */
+	UFUNCTION()
+	virtual void UpdateInteractingState(bool bInteracting, uint32 InPartIdentifier) = 0;
+
+	/*
+	 * Updates the hittable state of the specified gizmo part. 
+	 */
+	UFUNCTION()
+	virtual void UpdateHittableState(bool bHittable, uint32 InPartIdentifier) = 0;
+
+};
+
+
+
+
+UINTERFACE()
+class INTERACTIVETOOLSFRAMEWORK_API UGizmoRenderTarget : public UInterface
+{
+	GENERATED_BODY()
+};
+
+/**
+ * UGizmoRenderTarget is an interface used to provide rendering of a target
+ */
+class INTERACTIVETOOLSFRAMEWORK_API IGizmoRenderTarget
+{
+	GENERATED_BODY()
+public:
+	/**
+	 * Renders the target using the current tools context.
+	 */
+	virtual void Render(IToolsContextRenderAPI* RenderAPI) const = 0;
+};
+
+
+
+
+UINTERFACE()
+class INTERACTIVETOOLSFRAMEWORK_API UGizmoRenderMultiTarget : public UInterface
+{
+	GENERATED_BODY()
+};
+
+/**
+ * IGizmoRenderMultiTarget is an interface used to provide rendering of a target and the 
+ * ability to specify which part of a target should be visible.
+ * 
+ * For a gizmo with multiple parts, the part identifier establishes a correspondence between a gizmo part 
+ * and the elements representing that part within the hit target. The valid part identifiers should 
+ * be defined in the gizmo. Identifier 0 is reserved for the default ID which should be assigned to 
+ * elements that do not correspond to any gizmo part, such as non-hittable decorative elements.
+ */
+class INTERACTIVETOOLSFRAMEWORK_API IGizmoRenderMultiTarget
+{
+	GENERATED_BODY()
+public:
+	/**
+	 * Renders the target using the current tools context.
+	 */
+	virtual void Render(IToolsContextRenderAPI* RenderAPI) const = 0;
 
+	/*
+	 * Updates the visibility state of the specified gizmo part.
+	 */
+	UFUNCTION()
+	virtual void UpdateVisibilityState(bool bVisible, uint32 InPartIdentifier) = 0;
+};
 
 
 UINTERFACE()
diff --git a/Engine/Source/Runtime/InteractiveToolsFramework/Public/InputState.h b/Engine/Source/Runtime/InteractiveToolsFramework/Public/InputState.h
--- a/Engine/Source/Runtime/InteractiveToolsFramework/Public/InputState.h
+++ b/Engine/Source/Runtime/InteractiveToolsFramework/Public/InputState.h
@@ -52,12 +52,28 @@ struct INTERACTIVETOOLSFRAMEWORK_API FInputRayHit
 	/** True if HitNormal was set */
 	bool bHasHitNormal;
 
-	/** client-defined integer identifier for hit object/element/target/etc */
+	/** Client-defined integer identifier for hit object/element/target/etc */
 	int32 HitIdentifier;
 
-	/** client-defined pointer for hit object/target */
+	/**
+	  * Client-defined pointer for hit object/element/target/etc. 
+	  * HitOwner and HitObject should be set to the same pointer if the HitOwner derives from UObject.
+	  */
 	void* HitOwner;
 
+	/**
+	  * Client-defined pointer for UObject-derived hit owners.  
+	  * HitOwner and HitObject should be set to the same pointer if the HitOwner derives from UObject. 
+	  */
+	TWeakObjectPtr<UObject> HitObject;
+
+	/** Set hit object, will also set hit owner to the same value */
+	void SetHitObject(UObject* InHitObject)
+	{
+		HitObject = InHitObject;
+		HitOwner = InHitObject;
+	}
+
 	FInputRayHit()
 	{
 		bHit = false;
@@ -66,6 +82,7 @@ struct INTERACTIVETOOLSFRAMEWORK_API FInputRayHit
 		bHasHitNormal = false;
 		HitIdentifier = 0;
 		HitOwner = nullptr;
+		HitObject = nullptr;
 	}
 
 	explicit FInputRayHit(float HitDepthIn)
@@ -76,6 +93,7 @@ struct INTERACTIVETOOLSFRAMEWORK_API FInputRayHit
 		bHasHitNormal = false;
 		HitIdentifier = 0;
 		HitOwner = nullptr;
+		HitObject = nullptr;
 	}
 
 	explicit FInputRayHit(float HitDepthIn, const FVector& HitNormalIn)
@@ -86,6 +104,7 @@ struct INTERACTIVETOOLSFRAMEWORK_API FInputRayHit
 		bHasHitNormal = true;
 		HitIdentifier = 0;
 		HitOwner = nullptr;
+		HitObject = nullptr;
 	}
 };
 
