Implemented gizmo element circle and rectangle line trace methods.

diff --git a/Engine/Source/Runtime/Core/Public/Math/Plane.h b/Engine/Source/Runtime/Core/Public/Math/Plane.h
--- a/Engine/Source/Runtime/Core/Public/Math/Plane.h
+++ b/Engine/Source/Runtime/Core/Public/Math/Plane.h
@@ -638,6 +638,15 @@ inline UE::Math::TVector<T> FMath::RayPlaneIntersection(const UE::Math::TVector<
 	return RayOrigin + RayDirection * Distance;
 }
 
+template<typename T>
+inline T FMath::RayPlaneIntersectionParam(const UE::Math::TVector<T>& RayOrigin, const UE::Math::TVector<T>& RayDirection, const UE::Math::TPlane<T>& Plane)
+{
+	using TVector = UE::Math::TVector<T>;
+	const TVector PlaneNormal = TVector(Plane.X, Plane.Y, Plane.Z);
+	const TVector PlaneOrigin = PlaneNormal * Plane.W;
+
+	return TVector::DotProduct((PlaneOrigin - RayOrigin), PlaneNormal) / TVector::DotProduct(RayDirection, PlaneNormal);
+}
 
 template<typename T>
 inline UE::Math::TVector<T> FMath::LinePlaneIntersection
diff --git a/Engine/Source/Runtime/Core/Public/Math/UnrealMathUtility.h b/Engine/Source/Runtime/Core/Public/Math/UnrealMathUtility.h
--- a/Engine/Source/Runtime/Core/Public/Math/UnrealMathUtility.h
+++ b/Engine/Source/Runtime/Core/Public/Math/UnrealMathUtility.h
@@ -1740,6 +1740,19 @@ public:
 	template<typename FReal>
 	UE_NODISCARD static UE::Math::TVector<FReal> RayPlaneIntersection(const UE::Math::TVector<FReal>& RayOrigin, const UE::Math::TVector<FReal>& RayDirection, const UE::Math::TPlane<FReal>& Plane);
 
+	/**
+	 * Find the intersection of a ray and a plane.  The ray has a start point with an infinite length.  Assumes that the
+	 * line and plane do indeed intersect; you must make sure they're not parallel before calling.
+	 *
+	 * @param RayOrigin	The start point of the ray
+	 * @param RayDirection	The direction the ray is pointing (normalized vector)
+	 * @param Plane	The plane to intersect with
+	 *
+	 * @return The distance parameter along ray of the point of intersection between the ray and the plane.
+	 */
+	template<typename FReal>
+	UE_NODISCARD static FReal RayPlaneIntersectionParam(const UE::Math::TVector<FReal>& RayOrigin, const UE::Math::TVector<FReal>& RayDirection, const UE::Math::TPlane<FReal>& Plane);
+
 	/**
 	 * Find the intersection of a line and an offset plane. Assumes that the
 	 * line and plane do indeed intersect; you must make sure they're not
diff --git a/Engine/Source/Runtime/GeometryCore/Public/Intersection/IntersectionUtil.h b/Engine/Source/Runtime/GeometryCore/Public/Intersection/IntersectionUtil.h
--- a/Engine/Source/Runtime/GeometryCore/Public/Intersection/IntersectionUtil.h
+++ b/Engine/Source/Runtime/GeometryCore/Public/Intersection/IntersectionUtil.h
@@ -271,7 +271,54 @@ namespace IntersectionUtil
 		return result;
 	}
 
-};
+	 /**
+	  * Intersect ray with circle and return intersection info (# hits, ray parameters)
+	  * @return true if ray intersects circle
+	  */
+	template <typename RealType>
+	bool RayCircleIntersection(
+		const TVector<RealType>& RayOrigin,
+		const TVector<RealType>& RayDirection,
+		const TVector<RealType>& CircleCenter,
+		RealType CircleRadius,
+		const TVector<RealType>& CircleNormal,
+		FLinearIntersection& Result)
+	{
 
+		// if ray is parallel to circle, no hit
+		if (FMath::IsNearlyZero(TVector<RealType>::DotProduct(CircleNormal, RayDirection)))
+		{
+			Result.intersects = false;
+			Result.numIntersections = 0;
+			return false;
+		}
 
+		TPlane<RealType> Plane(CircleCenter, CircleNormal);
+		RealType Param = FMath::RayPlaneIntersectionParam(RayOrigin, RayDirection, Plane);
 
+		if (Param < 0)
+		{
+			Result.intersects = false;
+			Result.numIntersections = 0;
+			return false;
+		}
+
+		TVector<RealType> HitPoint = RayOrigin + RayDirection * Param;
+		TVector<RealType> Offset = HitPoint - CircleCenter;
+		const RealType CircleRadiusSquared = CircleRadius * CircleRadius;
+
+		if (Offset.SizeSquared() <= CircleRadiusSquared)
+		{
+			Result.intersects = true;
+			Result.numIntersections = 1;
+			Result.parameter.Min = Result.parameter.Max = Param;
+		}
+		else
+		{
+			Result.intersects = false;
+			Result.numIntersections = 0;
+		}
+
+		return Result.intersects;
+	}
+};
diff --git a/Engine/Source/Runtime/InteractiveToolsFramework/InteractiveToolsFramework.Build.cs b/Engine/Source/Runtime/InteractiveToolsFramework/InteractiveToolsFramework.Build.cs
--- a/Engine/Source/Runtime/InteractiveToolsFramework/InteractiveToolsFramework.Build.cs
+++ b/Engine/Source/Runtime/InteractiveToolsFramework/InteractiveToolsFramework.Build.cs
@@ -40,6 +40,7 @@ public class InteractiveToolsFramework : ModuleRules
 			{
 				"Engine",
                 "RHI",
+				"GeometryCore"
 				//"Slate",
 				//"SlateCore",
 				// ... add private dependencies that you statically link with here ...	
diff --git a/Engine/Source/Runtime/InteractiveToolsFramework/Private/BaseGizmos/GizmoElementCircle.cpp b/Engine/Source/Runtime/InteractiveToolsFramework/Private/BaseGizmos/GizmoElementCircle.cpp
--- a/Engine/Source/Runtime/InteractiveToolsFramework/Private/BaseGizmos/GizmoElementCircle.cpp
+++ b/Engine/Source/Runtime/InteractiveToolsFramework/Private/BaseGizmos/GizmoElementCircle.cpp
@@ -4,6 +4,7 @@
 #include "BaseGizmos/GizmoRenderingUtil.h"
 #include "BaseGizmos/GizmoMath.h"
 #include "InputState.h"
+#include "Intersection/IntersectionUtil.h"
 #include "Materials/MaterialInterface.h"
 #include "SceneManagement.h"
 
@@ -73,7 +74,30 @@ FInputRayHit UGizmoElementCircle::LineTrace(const FVector RayOrigin, const FVect
 {
 	if (IsHittableInView())
 	{
-		// @todo - implement ray-circle intersection 
+		if (bHitMesh)
+		{
+			bool bIntersects = false;
+			double Param = 0.0;
+
+			const FVector WorldNormal = CachedLocalToWorldTransform.TransformVectorNoScale(Normal);
+			const FVector WorldCenter = CachedLocalToWorldTransform.TransformPosition(Center);
+			const double WorldRadius = CachedLocalToWorldTransform.GetScale3D().X * Radius;
+
+			UE::Geometry::FLinearIntersection Result;
+			IntersectionUtil::RayCircleIntersection(RayOrigin, RayDirection, WorldCenter, WorldRadius, WorldNormal, Result);
+			
+			if (Result.intersects)
+			{
+				FInputRayHit RayHit(Result.parameter.Min);
+				RayHit.SetHitObject(this);
+				RayHit.HitIdentifier = PartIdentifier;
+				return RayHit;
+			}
+		}
+		else if (bHitLine)
+		{
+			// @todo - add hit testing for line-drawn circle, requires storing PixelToWorld scale factor
+		}
 	}
 	return FInputRayHit();
 }
@@ -183,4 +207,4 @@ void UGizmoElementCircle::SetScreenSpace(bool InScreenSpace)
 bool UGizmoElementCircle::GetScreenSpace()
 {
 	return bScreenSpace;
-}
\ No newline at end of file
+}
diff --git a/Engine/Source/Runtime/InteractiveToolsFramework/Private/BaseGizmos/GizmoElementRectangle.cpp b/Engine/Source/Runtime/InteractiveToolsFramework/Private/BaseGizmos/GizmoElementRectangle.cpp
--- a/Engine/Source/Runtime/InteractiveToolsFramework/Private/BaseGizmos/GizmoElementRectangle.cpp
+++ b/Engine/Source/Runtime/InteractiveToolsFramework/Private/BaseGizmos/GizmoElementRectangle.cpp
@@ -6,6 +6,7 @@
 #include "InputState.h"
 #include "Materials/MaterialInterface.h"
 #include "SceneManagement.h"
+#include "Math/UnrealMathUtility.h"
 
 void UGizmoElementRectangle::Render(IToolsContextRenderAPI* RenderAPI, const FRenderTraversalState& RenderState)
 {
@@ -74,7 +75,48 @@ FInputRayHit UGizmoElementRectangle::LineTrace(const FVector RayOrigin, const FV
 {
 	if (IsHittableInView())
 	{
-		// @todo - implement ray-circle intersection
+		if (bHitMesh)
+		{
+			const FVector UpAxis = CachedLocalToWorldTransform.TransformVectorNoScale(UpDirection);
+			const FVector SideAxis = CachedLocalToWorldTransform.TransformVectorNoScale(SideDirection);
+			const FVector Normal = FVector::CrossProduct(UpAxis, SideAxis);
+			const FVector WorldCenter = CachedLocalToWorldTransform.TransformPosition(Center);
+			const double Scale = CachedLocalToWorldTransform.GetScale3D().X;
+			const double WorldHeight = Scale * Height;
+			const double WorldWidth = Scale * Width;
+			const FVector Base = WorldCenter - UpAxis * WorldHeight * 0.5 - SideAxis * WorldWidth * 0.5;
+
+			// if ray is parallel to rectangle, no hit
+			if (FMath::IsNearlyZero(FVector::DotProduct(Normal, RayDirection)))
+			{
+				return FInputRayHit();
+			}
+
+			FPlane Plane(Base, Normal);
+			double HitDepth = FMath::RayPlaneIntersectionParam(RayOrigin, RayDirection, Plane);
+			if (HitDepth < 0)
+			{
+				return FInputRayHit();
+			}
+
+			FVector HitPoint = RayOrigin + RayDirection * HitDepth;
+			FVector HitOffset = HitPoint - Base;
+			double HdU = FVector::DotProduct(HitOffset, UpAxis);
+			double HdS = FVector::DotProduct(HitOffset, SideAxis);
+
+			// clip to rectangle dimensions
+			if (HdU >= 0.0 && HdU <= WorldHeight && HdS >= 0.0 && HdS <= WorldWidth)
+			{
+				FInputRayHit RayHit(HitDepth);
+				RayHit.SetHitObject(this);
+				RayHit.HitIdentifier = PartIdentifier;
+				return RayHit;
+			}
+		}
+		else if (bHitLine)
+		{
+			// @todo - add hit testing for line-drawn rectangle, requires storing PixelToWorld scale factor
+		}
 	}
 	return FInputRayHit();
 }
@@ -98,7 +140,6 @@ FVector UGizmoElementRectangle::GetCenter() const
 void UGizmoElementRectangle::SetWidth(float InWidth)
 {
 	Width = InWidth;
-	
 }
 
 float UGizmoElementRectangle::GetWidth() const
