Fix max values, use TNumericLimits.

diff --git a/Engine/Source/Runtime/InteractiveToolsFramework/Private/BaseGizmos/GizmoMath.cpp b/Engine/Source/Runtime/InteractiveToolsFramework/Private/BaseGizmos/GizmoMath.cpp
--- a/Engine/Source/Runtime/InteractiveToolsFramework/Private/BaseGizmos/GizmoMath.cpp
+++ b/Engine/Source/Runtime/InteractiveToolsFramework/Private/BaseGizmos/GizmoMath.cpp
@@ -471,7 +471,7 @@ void GizmoMath::RayConeIntersection(
 				// intersection with the positive cone.
 				if (DdU > 0.0)
 				{
-					RayParam[1] = FLT_MAX;
+					RayParam[1] = TNumericLimits<RealType>::Max();
 				}
 				else
 				{
@@ -562,8 +562,8 @@ void GizmoMath::IntervalIntervalIntersection(
 	if (Interval0[1] < Interval1[0] || Interval0[0] > Interval1[1])
 	{
 		OutNumIntersections = 0;
-		OutResult0 = FLT_MAX;
-		OutResult1 = -FLT_MAX;
+		OutResult0 = TNumericLimits<RealType>::Max();
+		OutResult1 = -TNumericLimits<RealType>::Max();
 	}
 	else if (Interval0[1] > Interval1[0])
 	{
