Implemented gizmo box element line trace.

diff --git a/Engine/Source/Runtime/GeometryCore/Public/Intersection/IntrRay3OrientedBox3.h b/Engine/Source/Runtime/GeometryCore/Public/Intersection/IntrRay3OrientedBox3.h
new file mode 100644
--- /dev/null
+++ b/Engine/Source/Runtime/GeometryCore/Public/Intersection/IntrRay3OrientedBox3.h
@@ -0,0 +1,71 @@
+// Copyright Epic Games, Inc. All Rights Reserved.
+
+#pragma once
+
+#include "BoxTypes.h"
+#include "IntrRay3AxisAlignedBox3.h"
+#include "OrientedBoxTypes.h"
+#include "VectorUtil.h"
+
+namespace UE
+{
+namespace Geometry
+{
+
+using namespace UE::Math;
+
+/**
+ * Compute intersection between 3D ray and 3D oriented box
+ */
+template <typename RealType>
+class TIntrRay3OrientedBox3
+{
+public:
+
+	/**
+	 * Find intersection of ray with oriented box
+	 * @param Ray query ray
+	 * @param Box query box
+	 * @param RayParamOut ray intersect T-value, or TNumericLimits::Max() on miss
+	 * @return true if ray intersects box
+	 * 
+	 * Based on GeometricTools Engine implementation:
+	 *		https://www.geometrictools.com/GTE/Mathematics/IntrRay3OrientedBox3.h
+	 */
+
+	static bool FindIntersection(
+		const TRay<RealType>& Ray, 
+		const TOrientedBox3<RealType>& Box,
+		RealType& RayParamOut)
+	{
+		const TVector<RealType> OriginOffset = Ray.Origin - Box.Frame.Origin;
+
+		const TVector<RealType> Axis0 = Box.Frame.GetAxis(0);
+		const TVector<RealType> Axis1 = Box.Frame.GetAxis(1);
+		const TVector<RealType> Axis2 = Box.Frame.GetAxis(2);
+
+		const TVector<RealType> Origin = TVector<RealType>(
+			TVector<RealType>::DotProduct(OriginOffset, Axis0),
+			TVector<RealType>::DotProduct(OriginOffset, Axis1),
+			TVector<RealType>::DotProduct(OriginOffset, Axis2));
+
+		const TVector<RealType> Direction = TVector<RealType>(
+			TVector<RealType>::DotProduct(Ray.Direction, Axis0),
+			TVector<RealType>::DotProduct(Ray.Direction, Axis1),
+			TVector<RealType>::DotProduct(Ray.Direction, Axis2));
+
+		const TRay<RealType> AdjustedRay(Origin, Direction);
+
+		const TAxisAlignedBox3<RealType> AxisAlignedBox(
+			TVector<RealType>(-Box.Extents.X, -Box.Extents.Y, -Box.Extents.Z),
+			TVector<RealType>(Box.Extents.X, Box.Extents.Y, Box.Extents.Z));
+
+		return TIntrRay3AxisAlignedBox3<RealType>::FindIntersection(AdjustedRay, AxisAlignedBox, RayParamOut);
+	}
+};
+
+typedef TIntrRay3OrientedBox3<float> FIntrRay3OrientedBox3f;
+typedef TIntrRay3OrientedBox3<double> FIntrRay3OrientedBox3d;
+
+} // end namespace UE::Geometry
+} // end namespace UE
diff --git a/Engine/Source/Runtime/InteractiveToolsFramework/Private/BaseGizmos/GizmoElementBox.cpp b/Engine/Source/Runtime/InteractiveToolsFramework/Private/BaseGizmos/GizmoElementBox.cpp
--- a/Engine/Source/Runtime/InteractiveToolsFramework/Private/BaseGizmos/GizmoElementBox.cpp
+++ b/Engine/Source/Runtime/InteractiveToolsFramework/Private/BaseGizmos/GizmoElementBox.cpp
@@ -1,10 +1,10 @@
 // Copyright Epic Games, Inc. All Rights Reserved.
  
 #include "BaseGizmos/GizmoElementBox.h"
-#include "BaseGizmos/GizmoMath.h"
 #include "BaseGizmos/GizmoRenderingUtil.h"
 #include "Materials/MaterialInterface.h"
 #include "InputState.h"
+#include "Intersection/IntrRay3OrientedBox3.h"
 #include "SceneManagement.h"
 
 void UGizmoElementBox::Render(IToolsContextRenderAPI* RenderAPI, const FRenderTraversalState& RenderState)
@@ -42,39 +42,37 @@ void UGizmoElementBox::Render(IToolsContextRenderAPI* RenderAPI, const FRenderTr
 
 			FQuat Rotation = FRotationMatrix::MakeFromYZ(AdjustedSideDir, AdjustedUpDir).ToQuat();
 			FPrimitiveDrawInterface* PDI = RenderAPI->GetPrimitiveDrawInterface();
-			LocalToWorldTransform = FTransform(Rotation, Center) * LocalToWorldTransform;
+			FTransform DrawTransform = FTransform(Rotation, Center) * LocalToWorldTransform;
 			const FVector HalfDimensions = Dimensions * 0.5;
-			DrawBox(PDI, LocalToWorldTransform.ToMatrixWithScale(), HalfDimensions, UseMaterial->GetRenderProxy(), SDPG_Foreground);
+			DrawBox(PDI, DrawTransform.ToMatrixWithScale(), HalfDimensions, UseMaterial->GetRenderProxy(), SDPG_Foreground);
 		}
 	}
 
 	CacheRenderState(LocalToWorldTransform, bVisibleViewDependent);
 }
 
-
-
-FInputRayHit UGizmoElementBox::LineTrace(const FVector RayCenter, const FVector RayDirection)
+FInputRayHit UGizmoElementBox::LineTrace(const FVector RayOrigin, const FVector RayDirection)
 {
 	if (IsHittableInView())
 	{
-		// @todo - modify ray-box intersection to work with updated properties
-
-#if 0
-		bool bIntersects = false;
-		float RayParam = 0.0f;
-
-
-		const FVector BoxOrigin = Origin + Offset;
-		GizmoMath::RayBoxIntersection(
-			BoxOrigin, CubeObject->Axis, CubeObject->Radius, CubeObject->Height,
-			ClickPos.WorldRay.Origin, ClickPos.WorldRay.Direction,
-			bIntersects, RayParam);
-
-		if (bIntersects)
+		const FVector YAxis = CachedLocalToWorldTransform.TransformVectorNoScale(SideDirection);
+		const FVector ZAxis = CachedLocalToWorldTransform.TransformVectorNoScale(UpDirection);
+		const FVector XAxis = FVector::CrossProduct(YAxis, ZAxis);
+		const FVector WorldCenter = CachedLocalToWorldTransform.TransformPosition(Center);
+		const double Scale = CachedLocalToWorldTransform.GetScale3D().X;
+		const FVector WorldExtent = Dimensions * Scale * 0.5;
+
+		double HitDepth = 0.0;
+		UE::Geometry::TRay<double> Ray(RayOrigin, RayDirection);
+		UE::Geometry::TFrame3<double> Frame(WorldCenter, XAxis, YAxis, ZAxis);
+		UE::Geometry::TOrientedBox3<double> Box(Frame, WorldExtent);
+		if (UE::Geometry::TIntrRay3OrientedBox3<double>::FindIntersection(Ray, Box, HitDepth))
 		{
-			return FInputRayHit(RayParam);
+			FInputRayHit RayHit(HitDepth);
+			RayHit.SetHitObject(this);
+			RayHit.HitIdentifier = PartIdentifier;
+			return RayHit;
 		}
-#endif
 	}
 
 	return FInputRayHit();
