Fix comments.

diff --git a/Engine/Source/Runtime/InteractiveToolsFramework/Public/BaseGizmos/GizmoElementBox.h b/Engine/Source/Runtime/InteractiveToolsFramework/Public/BaseGizmos/GizmoElementBox.h
--- a/Engine/Source/Runtime/InteractiveToolsFramework/Public/BaseGizmos/GizmoElementBox.h
+++ b/Engine/Source/Runtime/InteractiveToolsFramework/Public/BaseGizmos/GizmoElementBox.h
@@ -52,11 +52,11 @@ protected:
 	UPROPERTY()
 	FVector Dimensions = FVector(1.0f, 1.0f, 1.0f);
 
-	// Box up axis
+	// Box up axis, corresponds to Z dimension
 	UPROPERTY()
 	FVector UpDirection = FVector::UpVector;
 
-	// Box side axis
+	// Box side axis, corresponds to Y dimension
 	UPROPERTY()
 	FVector SideDirection = FVector::RightVector;
 };
diff --git a/Engine/Source/Runtime/InteractiveToolsFramework/Public/BaseGizmos/GizmoMath.h b/Engine/Source/Runtime/InteractiveToolsFramework/Public/BaseGizmos/GizmoMath.h
--- a/Engine/Source/Runtime/InteractiveToolsFramework/Public/BaseGizmos/GizmoMath.h
+++ b/Engine/Source/Runtime/InteractiveToolsFramework/Public/BaseGizmos/GizmoMath.h
@@ -62,7 +62,6 @@ namespace GizmoMath
 		bool& bIntersectsOut, FVector& SphereIntersectionPointOut);
 
 	/**
-	 * @todo change FVector to a templated vector type specialized by RealType (double or float)
 	 * Find the intersection of the ray defined by RayOrigin and RayDirection
 	 * with the cylinder defined by CylinderCenter, CylinderAxis, CylinderRadius, and CylinderHeight.
 	 * Returns intersection success/failure in bOutIntersects and the hit depth parameter 
@@ -76,7 +75,6 @@ namespace GizmoMath
 
 
 	/**
-	 * @todo change FVector to a templated vector type specialized by RealType (double or float)
 	 * Find the intersection of the ray defined by RayOrigin and RayDirection
 	 * with the cone defined by ConeCenter, ConeDirection, ConeCosAngle, and ConeHeight.
 	 * Where ConeCosAngle is the cosine of the cone opening angle, defined as the angle
