Added pixel hit distance threshold support to gizmo element library.

diff --git a/Engine/Source/Editor/Experimental/EditorInteractiveToolsFramework/Private/EditorGizmos/TransformGizmo.cpp b/Engine/Source/Editor/Experimental/EditorInteractiveToolsFramework/Private/EditorGizmos/TransformGizmo.cpp
--- a/Engine/Source/Editor/Experimental/EditorInteractiveToolsFramework/Private/EditorGizmos/TransformGizmo.cpp
+++ b/Engine/Source/Editor/Experimental/EditorInteractiveToolsFramework/Private/EditorGizmos/TransformGizmo.cpp
@@ -126,8 +126,7 @@ void UTransformGizmo::Render(IToolsContextRenderAPI* RenderAPI)
 		CurrentTransform = ActiveTarget->GetTransform();
 
 		UGizmoElementBase::FRenderTraversalState RenderState;
-		RenderState.LocalToWorldTransform = GetGizmoTransform();
-
+		RenderState.Initialize(RenderAPI->GetSceneView(), GetGizmoTransform());
 		GizmoElementRoot->Render(RenderAPI, RenderState);
 	}
 }
@@ -896,6 +895,7 @@ UGizmoElementRectangle* UTransformGizmo::MakeTranslateScreenSpaceHandle()
 	RectangleElement->SetHitMesh(true);
 	RectangleElement->SetDrawMesh(false);
 	RectangleElement->SetDrawLine(true);
+	RectangleElement->SetHoverLineThicknessMultiplier(3.0);
 	return RectangleElement;
 }
 
diff --git a/Engine/Source/Runtime/InteractiveToolsFramework/Private/BaseGizmos/GizmoElementArrow.cpp b/Engine/Source/Runtime/InteractiveToolsFramework/Private/BaseGizmos/GizmoElementArrow.cpp
--- a/Engine/Source/Runtime/InteractiveToolsFramework/Private/BaseGizmos/GizmoElementArrow.cpp
+++ b/Engine/Source/Runtime/InteractiveToolsFramework/Private/BaseGizmos/GizmoElementArrow.cpp
@@ -24,6 +24,16 @@ void UGizmoElementArrow::Render(IToolsContextRenderAPI* RenderAPI, const FRender
 		return;
 	}
 
+	if (bUpdateArrowBody)
+	{
+		UpdateArrowBody();
+	}
+
+	if (bUpdateArrowHead)
+	{
+		UpdateArrowHead();
+	}
+
 	const FSceneView* View = RenderAPI->GetSceneView();
 	bool bVisibleViewDependent = GetViewDependentVisibility(View, RenderState.LocalToWorldTransform, Base);
 
@@ -59,7 +69,7 @@ void UGizmoElementArrow::Render(IToolsContextRenderAPI* RenderAPI, const FRender
 		}
 	}
 
-	CacheRenderState(RenderStateCopy.LocalToWorldTransform, bVisibleViewDependent);
+	CacheRenderState(RenderStateCopy.LocalToWorldTransform, RenderStateCopy.PixelToWorldScale, bVisibleViewDependent);
 }
 
 FInputRayHit UGizmoElementArrow::LineTrace(const FVector RayOrigin, const FVector RayDirection)
@@ -103,9 +113,12 @@ FBoxSphereBounds UGizmoElementArrow::CalcBounds(const FTransform& LocalToWorld)
 
 void UGizmoElementArrow::SetBase(const FVector& InBase)
 {
-	Base = InBase;
-	UpdateArrowBody();
-	UpdateArrowHead();
+	if (Base != InBase)
+	{
+		Base = InBase;
+		bUpdateArrowBody = true;
+		bUpdateArrowHead = true;
+	}
 }
 
 FVector UGizmoElementArrow::GetBase() const
@@ -117,8 +130,8 @@ void UGizmoElementArrow::SetDirection(const FVector& InDirection)
 {
 	Direction = InDirection;
 	Direction.Normalize();
-	UpdateArrowBody();
-	UpdateArrowHead();
+	bUpdateArrowBody = true;
+	bUpdateArrowHead = true;
 }
 
 FVector UGizmoElementArrow::GetDirection() const
@@ -130,7 +143,7 @@ void UGizmoElementArrow::SetSideDirection(const FVector& InSideDirection)
 {
 	SideDirection = InSideDirection;
 	SideDirection.Normalize();
-	UpdateArrowHead();
+	bUpdateArrowHead = true;
 }
 
 FVector UGizmoElementArrow::GetSideDirection() const
@@ -140,9 +153,12 @@ FVector UGizmoElementArrow::GetSideDirection() const
 
 void UGizmoElementArrow::SetBodyLength(float InBodyLength)
 {
-	BodyLength = InBodyLength;
-	UpdateArrowBody();
-	UpdateArrowHead();
+	if (BodyLength != InBodyLength)
+	{
+		BodyLength = InBodyLength;
+		bUpdateArrowBody = true;
+		bUpdateArrowHead = true;
+	}
 }
 
 float UGizmoElementArrow::GetBodyLength() const
@@ -152,9 +168,12 @@ float UGizmoElementArrow::GetBodyLength() const
 
 void UGizmoElementArrow::SetBodyRadius(float InBodyRadius)
 {
-	BodyRadius = InBodyRadius;
-	UpdateArrowBody();
-	UpdateArrowHead();
+	if (BodyRadius != InBodyRadius)
+	{
+		BodyRadius = InBodyRadius;
+		bUpdateArrowBody = true;
+		bUpdateArrowHead = true;
+	}
 }
 
 float UGizmoElementArrow::GetBodyRadius() const
@@ -164,8 +183,11 @@ float UGizmoElementArrow::GetBodyRadius() const
 
 void UGizmoElementArrow::SetHeadLength(float InHeadLength)
 {
-	HeadLength = InHeadLength;
-	UpdateArrowHead();
+	if (HeadLength != InHeadLength)
+	{
+		HeadLength = InHeadLength;
+		bUpdateArrowHead = true;
+	}
 }
 
 float UGizmoElementArrow::GetHeadLength() const
@@ -175,8 +197,11 @@ float UGizmoElementArrow::GetHeadLength() const
 
 void UGizmoElementArrow::SetHeadRadius(float InHeadRadius)
 {
-	HeadRadius = InHeadRadius;
-	UpdateArrowHead();
+	if (HeadRadius != InHeadRadius)
+	{
+		HeadRadius = InHeadRadius;
+		bUpdateArrowHead = true;
+	}
 }
 
 float UGizmoElementArrow::GetHeadRadius() const
@@ -186,9 +211,12 @@ float UGizmoElementArrow::GetHeadRadius() const
 
 void UGizmoElementArrow::SetNumSides(int32 InNumSides)
 {
-	NumSides = InNumSides;
-	UpdateArrowBody();
-	UpdateArrowHead();
+	if (NumSides != InNumSides)
+	{
+		NumSides = InNumSides;
+		bUpdateArrowBody = true;
+		bUpdateArrowHead = true;
+	}
 }
 
 int32 UGizmoElementArrow::GetNumSides() const
@@ -196,6 +224,16 @@ int32 UGizmoElementArrow::GetNumSides() const
 	return NumSides;
 }
 
+void UGizmoElementArrow::SetPixelHitDistanceThreshold(float InPixelHitDistanceThreshold)
+{
+	if (PixelHitDistanceThreshold != InPixelHitDistanceThreshold)
+	{
+		PixelHitDistanceThreshold = InPixelHitDistanceThreshold;
+		bUpdateArrowBody = true;
+		bUpdateArrowHead = true;
+	}
+}
+
 void UGizmoElementArrow::SetHeadType(EGizmoElementArrowHeadType InHeadType)
 {
 	if (InHeadType != HeadType)
@@ -228,6 +266,9 @@ void UGizmoElementArrow::UpdateArrowBody()
 	CylinderElement->SetHeight(BodyLength);
 	CylinderElement->SetNumSides(NumSides);
 	CylinderElement->SetRadius(BodyRadius);
+	CylinderElement->SetPixelHitDistanceThreshold(PixelHitDistanceThreshold);
+
+	bUpdateArrowBody = false;
 }
 
 void UGizmoElementArrow::UpdateArrowHead()
@@ -242,6 +283,7 @@ void UGizmoElementArrow::UpdateArrowHead()
 		ConeElement->SetRadius(HeadRadius);
 		ConeElement->SetNumSides(NumSides);
 		ConeElement->SetElementInteractionState(ElementInteractionState);
+		ConeElement->SetPixelHitDistanceThreshold(PixelHitDistanceThreshold);
 	}
 	else // (HeadType == EGizmoElementArrowHeadType::Cube)
 	{
@@ -251,5 +293,8 @@ void UGizmoElementArrow::UpdateArrowHead()
 		BoxElement->SetSideDirection(SideDirection);
 		BoxElement->SetDimensions(FVector(HeadLength, HeadLength, HeadLength));
 		BoxElement->SetElementInteractionState(ElementInteractionState);
+		BoxElement->SetPixelHitDistanceThreshold(PixelHitDistanceThreshold);
 	}
+
+	bUpdateArrowHead = false;
 }
diff --git a/Engine/Source/Runtime/InteractiveToolsFramework/Private/BaseGizmos/GizmoElementBase.cpp b/Engine/Source/Runtime/InteractiveToolsFramework/Private/BaseGizmos/GizmoElementBase.cpp
--- a/Engine/Source/Runtime/InteractiveToolsFramework/Private/BaseGizmos/GizmoElementBase.cpp
+++ b/Engine/Source/Runtime/InteractiveToolsFramework/Private/BaseGizmos/GizmoElementBase.cpp
@@ -5,7 +5,7 @@
 #include "BaseGizmos/GizmoRenderingUtil.h"
 #include "Materials/MaterialInterface.h"
 
-DEFINE_LOG_CATEGORY_STATIC(LogGizmoElements, Log, All);
+DEFINE_LOG_CATEGORY_STATIC(LogGizmoElementBase, Log, All);
 
 bool UGizmoElementBase::GetViewDependentVisibility(const FSceneView* View, const FTransform& InLocalToWorldTransform, const FVector& InLocalCenter) const
 {
@@ -57,7 +57,7 @@ bool UGizmoElementBase::GetViewAlignRot(const FSceneView* View, const FTransform
 		bool bNonUniformScaleWarning = true;
 		if (bNonUniformScaleWarning)
 		{
-			UE_LOG(LogGizmoElements, Warning, TEXT("Gizmo element library view-dependent alignment does not currently support non-uniform scale (%f %f %f)."),
+			UE_LOG(LogGizmoElementBase, Warning, TEXT("Gizmo element library view-dependent alignment does not currently support non-uniform scale (%f %f %f)."),
 				Scale.X, Scale.Y, Scale.Z);
 			bNonUniformScaleWarning = false;
 		}
@@ -170,9 +170,10 @@ const UMaterialInterface* UGizmoElementBase::GetCurrentMaterial(const FRenderTra
 	return Material;
 }
 
-void UGizmoElementBase::CacheRenderState(const FTransform& InLocalToWorldState, bool InVisibleViewDependent)
+void UGizmoElementBase::CacheRenderState(const FTransform& InLocalToWorldState, double InPixelToWorldScale, bool InVisibleViewDependent)
 {
 	CachedLocalToWorldTransform = InLocalToWorldState;
+	CachedPixelToWorldScale = InPixelToWorldScale;
 	bHasCachedLocalToWorldTransform = true;
 	bCachedVisibleViewDependent = InVisibleViewDependent;
 }
@@ -180,6 +181,8 @@ void UGizmoElementBase::CacheRenderState(const FTransform& InLocalToWorldState,
 void UGizmoElementBase::ResetCachedRenderState()
 {
 	bHasCachedLocalToWorldTransform = false;
+	CachedLocalToWorldTransform = FTransform::Identity;
+	CachedPixelToWorldScale = 1.0f;
 	bCachedVisibleViewDependent = true;
 }
 
@@ -405,4 +408,13 @@ FColor UGizmoElementBase::GetVertexColor() const
 	return VertexColor;
 }
 
+void UGizmoElementBase::SetPixelHitDistanceThreshold(float InPixelHitDistanceThreshold)
+{
+	PixelHitDistanceThreshold = InPixelHitDistanceThreshold;
+}
+
+float UGizmoElementBase::GetPixelHitDistanceThreshold() const
+{
+	return PixelHitDistanceThreshold;
+}
 
diff --git a/Engine/Source/Runtime/InteractiveToolsFramework/Private/BaseGizmos/GizmoElementBox.cpp b/Engine/Source/Runtime/InteractiveToolsFramework/Private/BaseGizmos/GizmoElementBox.cpp
--- a/Engine/Source/Runtime/InteractiveToolsFramework/Private/BaseGizmos/GizmoElementBox.cpp
+++ b/Engine/Source/Runtime/InteractiveToolsFramework/Private/BaseGizmos/GizmoElementBox.cpp
@@ -48,7 +48,7 @@ void UGizmoElementBox::Render(IToolsContextRenderAPI* RenderAPI, const FRenderTr
 		}
 	}
 
-	CacheRenderState(LocalToWorldTransform, bVisibleViewDependent);
+	CacheRenderState(LocalToWorldTransform, RenderState.PixelToWorldScale, bVisibleViewDependent);
 }
 
 FInputRayHit UGizmoElementBox::LineTrace(const FVector RayOrigin, const FVector RayDirection)
diff --git a/Engine/Source/Runtime/InteractiveToolsFramework/Private/BaseGizmos/GizmoElementCircle.cpp b/Engine/Source/Runtime/InteractiveToolsFramework/Private/BaseGizmos/GizmoElementCircle.cpp
--- a/Engine/Source/Runtime/InteractiveToolsFramework/Private/BaseGizmos/GizmoElementCircle.cpp
+++ b/Engine/Source/Runtime/InteractiveToolsFramework/Private/BaseGizmos/GizmoElementCircle.cpp
@@ -64,10 +64,10 @@ void UGizmoElementCircle::Render(IToolsContextRenderAPI* RenderAPI, const FRende
 
 		if (bDrawLine)
 		{
-			DrawCircle(PDI, WorldCenter, Axis0, Axis1, LineColor, WorldRadius, NumSides, SDPG_Foreground, 0.0f);
+			DrawCircle(PDI, WorldCenter, Axis0, Axis1, LineColor, WorldRadius, NumSides, SDPG_Foreground, GetCurrentLineThickness());
 		}
 	}
-	CacheRenderState(LocalToWorldTransform, bVisibleViewDependent);
+	CacheRenderState(LocalToWorldTransform, RenderState.PixelToWorldScale, bVisibleViewDependent);
 }
 
 FInputRayHit UGizmoElementCircle::LineTrace(const FVector RayOrigin, const FVector RayDirection)
@@ -81,7 +81,8 @@ FInputRayHit UGizmoElementCircle::LineTrace(const FVector RayOrigin, const FVect
 
 			const FVector WorldNormal = CachedLocalToWorldTransform.TransformVectorNoScale(Normal);
 			const FVector WorldCenter = CachedLocalToWorldTransform.TransformPosition(Center);
-			const double WorldRadius = CachedLocalToWorldTransform.GetScale3D().X * Radius;
+			const double PixelHitThresholdAdjust = CachedPixelToWorldScale * PixelHitDistanceThreshold;
+			const double WorldRadius = CachedLocalToWorldTransform.GetScale3D().X * Radius + PixelHitThresholdAdjust;
 
 			UE::Geometry::FLinearIntersection Result;
 			IntersectionUtil::RayCircleIntersection(RayOrigin, RayDirection, WorldCenter, WorldRadius, WorldNormal, Result);
diff --git a/Engine/Source/Runtime/InteractiveToolsFramework/Private/BaseGizmos/GizmoElementCone.cpp b/Engine/Source/Runtime/InteractiveToolsFramework/Private/BaseGizmos/GizmoElementCone.cpp
--- a/Engine/Source/Runtime/InteractiveToolsFramework/Private/BaseGizmos/GizmoElementCone.cpp
+++ b/Engine/Source/Runtime/InteractiveToolsFramework/Private/BaseGizmos/GizmoElementCone.cpp
@@ -42,7 +42,7 @@ void UGizmoElementCone::Render(IToolsContextRenderAPI* RenderAPI, const FRenderT
 			DrawCone(PDI, RenderLocalToWorldTransform.ToMatrixWithScale(), Angle, Angle, NumSides, false, FColor::White, UseMaterial->GetRenderProxy(), SDPG_Foreground);
 		}
 	}
-	CacheRenderState(LocalToWorldTransform, bVisibleViewDependent);
+	CacheRenderState(LocalToWorldTransform, RenderState.PixelToWorldScale, bVisibleViewDependent);
 }
 
 
@@ -54,11 +54,12 @@ FInputRayHit UGizmoElementCone::LineTrace(const FVector RayOrigin, const FVector
 		bool bIntersects = false;
 		double RayParam = 0.0;
 
+		const double PixelHitThresholdAdjust = CachedPixelToWorldScale * PixelHitDistanceThreshold;
 		const double ConeSide = FMath::Sqrt(Height * Height + Radius * Radius);
 		const double CosAngle = Height / ConeSide;
-		const double ConeHeight = Height * CachedLocalToWorldTransform.GetScale3D().X;
-		const FVector ConeOrigin = CachedLocalToWorldTransform.TransformPosition(Origin);
+		const double ConeHeight = Height * CachedLocalToWorldTransform.GetScale3D().X + PixelHitThresholdAdjust * 2.0;
 		const FVector ConeDirection = CachedLocalToWorldTransform.TransformVectorNoScale(Direction);
+		const FVector ConeOrigin = CachedLocalToWorldTransform.TransformPosition(Origin) - ConeDirection * PixelHitThresholdAdjust;
 
 		GizmoMath::RayConeIntersection(
 			ConeOrigin,
diff --git a/Engine/Source/Runtime/InteractiveToolsFramework/Private/BaseGizmos/GizmoElementCylinder.cpp b/Engine/Source/Runtime/InteractiveToolsFramework/Private/BaseGizmos/GizmoElementCylinder.cpp
--- a/Engine/Source/Runtime/InteractiveToolsFramework/Private/BaseGizmos/GizmoElementCylinder.cpp
+++ b/Engine/Source/Runtime/InteractiveToolsFramework/Private/BaseGizmos/GizmoElementCylinder.cpp
@@ -41,7 +41,7 @@ void UGizmoElementCylinder::Render(IToolsContextRenderAPI* RenderAPI, const FRen
 
 		}
 	}
-	CacheRenderState(LocalToWorldTransform, bVisibleViewDependent);
+	CacheRenderState(LocalToWorldTransform, RenderState.PixelToWorldScale, bVisibleViewDependent);
 }
 
 
@@ -53,8 +53,9 @@ FInputRayHit UGizmoElementCylinder::LineTrace(const FVector RayOrigin, const FVe
 		bool bIntersects = false;
 		double RayParam = 0.0;
 		
-		const double CylinderHeight = Height * CachedLocalToWorldTransform.GetScale3D().X;
-		const double CylinderRadius = Radius * CachedLocalToWorldTransform.GetScale3D().X;
+		const double PixelHitThresholdAdjust = CachedPixelToWorldScale * PixelHitDistanceThreshold;
+		const double CylinderHeight = Height * CachedLocalToWorldTransform.GetScale3D().X + PixelHitThresholdAdjust * 2.0;
+		const double CylinderRadius = Radius * CachedLocalToWorldTransform.GetScale3D().X + PixelHitThresholdAdjust;
 		const FVector CylinderDirection = CachedLocalToWorldTransform.TransformVectorNoScale(Direction);
 		const FVector CylinderLocalCenter = Base + Direction * Height * 0.5;
 		const FVector CylinderCenter = CachedLocalToWorldTransform.TransformPosition(CylinderLocalCenter);
diff --git a/Engine/Source/Runtime/InteractiveToolsFramework/Private/BaseGizmos/GizmoElementGroup.cpp b/Engine/Source/Runtime/InteractiveToolsFramework/Private/BaseGizmos/GizmoElementGroup.cpp
--- a/Engine/Source/Runtime/InteractiveToolsFramework/Private/BaseGizmos/GizmoElementGroup.cpp
+++ b/Engine/Source/Runtime/InteractiveToolsFramework/Private/BaseGizmos/GizmoElementGroup.cpp
@@ -1,8 +1,6 @@
 // Copyright Epic Games, Inc. All Rights Reserved.
 
 #include "BaseGizmos/GizmoElementGroup.h"
-#include "BaseGizmos/GizmoRenderingUtil.h"
-
 
 void UGizmoElementGroup::Render(IToolsContextRenderAPI* RenderAPI, const FRenderTraversalState& RenderState)
 {
@@ -23,8 +21,7 @@ void UGizmoElementGroup::Render(IToolsContextRenderAPI* RenderAPI, const FRender
 		float Scale = RenderStateCopy.LocalToWorldTransform.GetScale3D().X;
 		if (bConstantScale)
 		{
-			const float ConstantScaleFactor = GizmoRenderingUtil::CalculateLocalPixelToWorldScale(RenderAPI->GetSceneView(), RenderStateCopy.LocalToWorldTransform.GetTranslation());
-			Scale *= ConstantScaleFactor;
+			Scale *= RenderStateCopy.PixelToWorldScale;
 		}
 		RenderStateCopy.LocalToWorldTransform.SetScale3D(FVector(Scale, Scale, Scale));
 
@@ -48,7 +45,7 @@ void UGizmoElementGroup::Render(IToolsContextRenderAPI* RenderAPI, const FRender
 		}
 	}
 
-	CacheRenderState(RenderStateCopy.LocalToWorldTransform, bVisibleViewDependent);
+	CacheRenderState(RenderStateCopy.LocalToWorldTransform, RenderStateCopy.PixelToWorldScale, bVisibleViewDependent);
 }
 
 FInputRayHit UGizmoElementGroup::LineTrace(const FVector Start, const FVector Direction)
diff --git a/Engine/Source/Runtime/InteractiveToolsFramework/Private/BaseGizmos/GizmoElementLineBase.cpp b/Engine/Source/Runtime/InteractiveToolsFramework/Private/BaseGizmos/GizmoElementLineBase.cpp
new file mode 100644
--- /dev/null
+++ b/Engine/Source/Runtime/InteractiveToolsFramework/Private/BaseGizmos/GizmoElementLineBase.cpp
@@ -0,0 +1,59 @@
+// Copyright Epic Games, Inc. All Rights Reserved.
+ 
+#include "BaseGizmos/GizmoElementLineBase.h"
+#include "BaseGizmos/GizmoElementBase.h"
+
+DEFINE_LOG_CATEGORY_STATIC(LogGizmoElementLineBase, Log, All);
+
+float UGizmoElementLineBase::GetCurrentLineThickness() const
+{
+	if (ElementInteractionState == EGizmoElementInteractionState::Hovering)
+	{
+		return (LineThickness == 0.0f ? HoverLineThicknessMultiplier : LineThickness * HoverLineThicknessMultiplier);
+	}
+	else if (ElementInteractionState == EGizmoElementInteractionState::Interacting)
+	{
+		return (LineThickness == 0.0f ? InteractLineThicknessMultiplier : LineThickness * InteractLineThicknessMultiplier);
+	}
+
+	return LineThickness;
+}
+
+void UGizmoElementLineBase::SetLineThickness(float InLineThickness)
+{
+	if (InLineThickness < 0.0)
+	{
+		UE_LOG(LogGizmoElementLineBase, Warning, TEXT("Invalid gizmo element line thickness %f, will be set to 0.0."), InLineThickness);
+		LineThickness = 0.0;
+	}
+	else
+	{
+		LineThickness = InLineThickness;
+	}
+}
+
+float UGizmoElementLineBase::GetLineThickness() const
+{
+	return LineThickness;
+}
+
+void UGizmoElementLineBase::SetHoverLineThicknessMultiplier(float InHoverLineThicknessMultiplier)
+{
+	HoverLineThicknessMultiplier = InHoverLineThicknessMultiplier;
+}
+
+float UGizmoElementLineBase::GetHoverLineThicknessMultiplier() const
+{
+	return HoverLineThicknessMultiplier;
+}
+
+void UGizmoElementLineBase::SetInteractLineThicknessMultiplier(float InInteractLineThicknessMultiplier)
+{
+	InteractLineThicknessMultiplier = InInteractLineThicknessMultiplier;
+}
+
+float UGizmoElementLineBase::GetInteractLineThicknessMultiplier() const
+{
+	return InteractLineThicknessMultiplier;
+}
+
diff --git a/Engine/Source/Runtime/InteractiveToolsFramework/Private/BaseGizmos/GizmoElementRectangle.cpp b/Engine/Source/Runtime/InteractiveToolsFramework/Private/BaseGizmos/GizmoElementRectangle.cpp
--- a/Engine/Source/Runtime/InteractiveToolsFramework/Private/BaseGizmos/GizmoElementRectangle.cpp
+++ b/Engine/Source/Runtime/InteractiveToolsFramework/Private/BaseGizmos/GizmoElementRectangle.cpp
@@ -65,10 +65,10 @@ void UGizmoElementRectangle::Render(IToolsContextRenderAPI* RenderAPI, const FRe
 		}
 		if (bDrawLine)
 		{
-			DrawRectangle(PDI, WorldCenter, Axis0, Axis1, LineColor, WorldWidth, WorldHeight, SDPG_Foreground, 0.0f);
+			DrawRectangle(PDI, WorldCenter, Axis0, Axis1, LineColor, WorldWidth, WorldHeight, SDPG_Foreground, GetCurrentLineThickness());
 		}
 	}
-	CacheRenderState(LocalToWorldTransform, bVisibleViewDependent);
+	CacheRenderState(LocalToWorldTransform, RenderState.PixelToWorldScale, bVisibleViewDependent);
 }
 
 FInputRayHit UGizmoElementRectangle::LineTrace(const FVector RayOrigin, const FVector RayDirection)
@@ -82,8 +82,9 @@ FInputRayHit UGizmoElementRectangle::LineTrace(const FVector RayOrigin, const FV
 			const FVector Normal = FVector::CrossProduct(UpAxis, SideAxis);
 			const FVector WorldCenter = CachedLocalToWorldTransform.TransformPosition(Center);
 			const double Scale = CachedLocalToWorldTransform.GetScale3D().X;
-			const double WorldHeight = Scale * Height;
-			const double WorldWidth = Scale * Width;
+			const double PixelHitThresholdAdjust = CachedPixelToWorldScale * PixelHitDistanceThreshold;
+			const double WorldHeight = Scale * Height + 2.0 * PixelHitThresholdAdjust;
+			const double WorldWidth = Scale * Width + 2.0 * PixelHitThresholdAdjust;
 			const FVector Base = WorldCenter - UpAxis * WorldHeight * 0.5 - SideAxis * WorldWidth * 0.5;
 
 			// if ray is parallel to rectangle, no hit
@@ -182,7 +183,7 @@ void UGizmoElementRectangle::SetScreenSpace(bool InScreenSpace)
 	bScreenSpace = InScreenSpace;
 }
 
-bool UGizmoElementRectangle::GetScreenSpace()
+bool UGizmoElementRectangle::GetScreenSpace() const
 {
 	return bScreenSpace;
 }
@@ -192,7 +193,7 @@ void UGizmoElementRectangle::SetLineColor(const FColor& InLineColor)
 	LineColor = InLineColor;
 }
 
-FColor UGizmoElementRectangle::GetLineColor()
+FColor UGizmoElementRectangle::GetLineColor() const
 {
 	return LineColor;
 }
diff --git a/Engine/Source/Runtime/InteractiveToolsFramework/Public/BaseGizmos/GizmoElementArrow.h b/Engine/Source/Runtime/InteractiveToolsFramework/Public/BaseGizmos/GizmoElementArrow.h
--- a/Engine/Source/Runtime/InteractiveToolsFramework/Public/BaseGizmos/GizmoElementArrow.h
+++ b/Engine/Source/Runtime/InteractiveToolsFramework/Public/BaseGizmos/GizmoElementArrow.h
@@ -75,6 +75,9 @@ public:
 	virtual void SetHeadType(EGizmoElementArrowHeadType InHeadType);
 	virtual EGizmoElementArrowHeadType GetHeadType() const;
 
+	// Pixel hit distance threshold, element will be scaled enough to add this threshold when line-tracing.
+	virtual void SetPixelHitDistanceThreshold(float InPixelHitDistanceThreshold) override;
+
 protected:
 
 	// Update arrow cylinder based on parameters
@@ -83,6 +86,12 @@ protected:
 	// Update arrow cone or box based on parameters
 	virtual void UpdateArrowHead();
 
+	// Flag indicating body properties need to be updated prior to render
+	bool bUpdateArrowBody = true;
+
+	// Flag indicating head properties need to be updated prior to render
+	bool bUpdateArrowHead = true;
+
 	// Arrow cylinder body
 	UPROPERTY()
 	TObjectPtr<UGizmoElementCylinder> CylinderElement;
diff --git a/Engine/Source/Runtime/InteractiveToolsFramework/Public/BaseGizmos/GizmoElementBase.h b/Engine/Source/Runtime/InteractiveToolsFramework/Public/BaseGizmos/GizmoElementBase.h
--- a/Engine/Source/Runtime/InteractiveToolsFramework/Public/BaseGizmos/GizmoElementBase.h
+++ b/Engine/Source/Runtime/InteractiveToolsFramework/Public/BaseGizmos/GizmoElementBase.h
@@ -4,7 +4,9 @@
 
 #include "CoreMinimal.h"
 #include "BaseGizmos/GizmoInterfaces.h"
+#include "BaseGizmos/GizmoRenderingUtil.h"
 #include "InputState.h"
+#include "SceneView.h"
 #include "ToolContextInterfaces.h"
 #include "UObject/GCObject.h"
 #include "GizmoElementBase.generated.h"
@@ -47,7 +49,7 @@ enum class EGizmoElementViewAlignType
 };
 
 /**
- * Base class for simple objects intended to be used as part of 3D Gizmos.
+ * Base class for 2d and 3d primitive objects intended to be used as part of 3D Gizmos.
  * Contains common properties and utility functions.
  * This class does nothing by itself, use subclasses like UGizmoElementCylinder
  */
@@ -73,14 +75,8 @@ public:
 		// Note: non-uniform scale is not supported and the X scale element will be used for uniform scaling.
 		FTransform LocalToWorldTransform;
 
-		// LocalToWorld translation 
-		FVector LocalToWorldTranslation;
-
-		// LocalToWorld rotation 
-		FQuat LocalToWorldRotation;
-
-		// LocalToWorld uniform scale, non-uniform scale is not supported
-		double LocalToWorldScale;
+		// Pixel to world scale
+		double PixelToWorldScale = 1.0;
 
 		// Material, if set, overrides the element's material
 		TWeakObjectPtr<UMaterialInterface> Material = nullptr;
@@ -93,6 +89,12 @@ public:
 
 		// Interact state, if not equal to none, overrides the element's interact state
 		EGizmoElementInteractionState InteractionState = EGizmoElementInteractionState::None;
+
+		void Initialize(const FSceneView* InSceneView, FTransform InTransform)
+		{
+			LocalToWorldTransform = InTransform;
+			PixelToWorldScale = GizmoRenderingUtil::CalculateLocalPixelToWorldScale(InSceneView, InTransform.GetLocation());
+		}
 	};
 
 public:
@@ -209,6 +211,10 @@ public:
 	virtual void SetVertexColor(const FColor& InVertexColor);
 	FColor GetVertexColor() const;
 
+	// Pixel hit distance threshold, element will be scaled enough to add this threshold when line-tracing.
+	virtual void SetPixelHitDistanceThreshold(float InPixelHitDistanceThreshold);
+	virtual float GetPixelHitDistanceThreshold() const;
+
 protected:
 
 	// Render and LineTrace should only occur when bEnabled is true.
@@ -295,21 +301,17 @@ protected:
 	UPROPERTY()
 	FColor VertexColor = FColor::White;
 
-	// Render stores the last cached transform, used by line trace.
-	UPROPERTY()
-	FTransform CachedLocalToWorldTransform = FTransform::Identity;
-
-	// Render stores the last cached translation, used by line trace.
+	// Pixel hit distance threshold, element will be scaled enough to add this threshold when line-tracing.
 	UPROPERTY()
-	FVector CachedLocalToWorldTranslation = FVector::ZeroVector;
+	float PixelHitDistanceThreshold = 7.0;
 
-	// Render stores the last cached rotation, used by line trace.
+	// Render stores the last cached transform, used by line trace.
 	UPROPERTY()
-	FQuat CachedLocalToWorldRotation = FQuat::Identity;
+	FTransform CachedLocalToWorldTransform = FTransform::Identity;
 
-	// Render stores the last cached rotation, used by line trace.
+	// Render stores the last pixel to world scale, used by line trace.
 	UPROPERTY()
-	double CachedLocalToWorldScale = 1.0;
+	float CachedPixelToWorldScale = 1.0;
 
 	// Whether last transform has been cached.
 	UPROPERTY()
@@ -329,7 +331,7 @@ protected:
 
 protected:
 	// Cache render state during render traversal, to be used subsequently when line tracing.
-	virtual void CacheRenderState(const FTransform& InLocalToWorldState, bool InVisibleViewDependent = true);
+	virtual void CacheRenderState(const FTransform& InLocalToWorldState, double InPixelToWorldScale, bool InVisibleViewDependent = true);
 
 	// Update render traversal state based on properties in the current object, should always be called at the beginning of Render()
 	virtual void UpdateRenderTraversalState(FRenderTraversalState& InRenderTraversalState);
diff --git a/Engine/Source/Runtime/InteractiveToolsFramework/Public/BaseGizmos/GizmoElementCircle.h b/Engine/Source/Runtime/InteractiveToolsFramework/Public/BaseGizmos/GizmoElementCircle.h
--- a/Engine/Source/Runtime/InteractiveToolsFramework/Public/BaseGizmos/GizmoElementCircle.h
+++ b/Engine/Source/Runtime/InteractiveToolsFramework/Public/BaseGizmos/GizmoElementCircle.h
@@ -3,7 +3,7 @@
 #pragma once
 
 #include "CoreMinimal.h"
-#include "BaseGizmos/GizmoElementBase.h"
+#include "BaseGizmos/GizmoElementLineBase.h"
 #include "InputState.h"
 #include "UObject/ObjectMacros.h"
 #include "GizmoElementCircle.generated.h"
@@ -13,7 +13,7 @@
  * Draws a filled or line circle based on parameters.
  */
 UCLASS(Transient)
-class INTERACTIVETOOLSFRAMEWORK_API UGizmoElementCircle : public UGizmoElementBase
+class INTERACTIVETOOLSFRAMEWORK_API UGizmoElementCircle : public UGizmoElementLineBase
 {
 	GENERATED_BODY()
 
diff --git a/Engine/Source/Runtime/InteractiveToolsFramework/Public/BaseGizmos/GizmoElementLineBase.h b/Engine/Source/Runtime/InteractiveToolsFramework/Public/BaseGizmos/GizmoElementLineBase.h
new file mode 100644
--- /dev/null
+++ b/Engine/Source/Runtime/InteractiveToolsFramework/Public/BaseGizmos/GizmoElementLineBase.h
@@ -0,0 +1,47 @@
+// Copyright Epic Games, Inc. All Rights Reserved.
+
+#pragma once
+
+#include "CoreMinimal.h"
+#include "BaseGizmos/GizmoElementBase.h"
+#include "GizmoElementLineBase.generated.h"
+
+/**
+ * Base class for 2d and 3d primitive objects which support line drawing,
+ * intended to be used as part of 3D Gizmos.
+ */
+UCLASS(Transient, Abstract)
+class INTERACTIVETOOLSFRAMEWORK_API UGizmoElementLineBase : public UGizmoElementBase
+{
+	GENERATED_BODY()
+public:
+
+	// Get line thickness for based on current element interaction state
+	virtual float GetCurrentLineThickness() const;
+
+	// Line thickness when rendering lines, 0.0 is valid and will render thinnest line 
+	virtual void SetLineThickness(float InLineThickness);
+	virtual float GetLineThickness() const;
+
+	// Multiplier applied to line thickness when hovering
+	virtual void SetHoverLineThicknessMultiplier(float InHoverLineThicknessMultiplier);
+	virtual float GetHoverLineThicknessMultiplier() const;
+
+	// Multiplier applied to line thickness when interacting
+	virtual void SetInteractLineThicknessMultiplier(float InInteractLineThicknessMultiplier);
+	virtual float GetInteractLineThicknessMultiplier() const;
+
+protected:
+
+	// Line thickness when rendering lines, must be >= 0.0, value of 0.0 will render thinnest line 
+	UPROPERTY()
+	float LineThickness = 0.0;
+
+	// Multiplier applied to line thickness when hovering
+	UPROPERTY(EditAnywhere, Category = Options)
+	float HoverLineThicknessMultiplier = 2.0f;
+
+	// Multiplier applied to line thickness when interacting
+	UPROPERTY(EditAnywhere, Category = Options)
+	float InteractLineThicknessMultiplier = 2.0f;
+};
\ No newline at end of file
diff --git a/Engine/Source/Runtime/InteractiveToolsFramework/Public/BaseGizmos/GizmoElementRectangle.h b/Engine/Source/Runtime/InteractiveToolsFramework/Public/BaseGizmos/GizmoElementRectangle.h
--- a/Engine/Source/Runtime/InteractiveToolsFramework/Public/BaseGizmos/GizmoElementRectangle.h
+++ b/Engine/Source/Runtime/InteractiveToolsFramework/Public/BaseGizmos/GizmoElementRectangle.h
@@ -3,7 +3,7 @@
 #pragma once
 
 #include "CoreMinimal.h"
-#include "BaseGizmos/GizmoElementBase.h"
+#include "BaseGizmos/GizmoElementLineBase.h"
 #include "InputState.h"
 #include "UObject/ObjectMacros.h"
 #include "GizmoElementRectangle.generated.h"
@@ -16,7 +16,7 @@ class FMaterialRenderProxy;
  * Draws a rectangle based on parameters.
  */
 UCLASS(Transient)
-class INTERACTIVETOOLSFRAMEWORK_API UGizmoElementRectangle : public UGizmoElementBase
+class INTERACTIVETOOLSFRAMEWORK_API UGizmoElementRectangle : public UGizmoElementLineBase
 {
 	GENERATED_BODY()
 
@@ -50,11 +50,11 @@ public:
 
 	// Screen space, when true orients the rectangle to the screen up and side vectors
 	virtual void SetScreenSpace(bool InScreenSpace);
-	virtual bool GetScreenSpace();
+	virtual bool GetScreenSpace() const;
 
 	// Line color
 	virtual void SetLineColor(const FColor& InLineColor);
-	virtual FColor GetLineColor();
+	virtual FColor GetLineColor() const;
 
 	// Draw mesh
 	virtual void SetDrawMesh(bool InDrawMesh);
