InteractiveToolsFramework: Fix UInteractiveToolsContext::DeactivateAllActiveTools always sending cancel to complete-style tools regardless of shutdown type. This avoids cancelling those tools (CubeGrid) when starting PIE.

diff --git a/Engine/Source/Runtime/InteractiveToolsFramework/Private/InteractiveToolsContext.cpp b/Engine/Source/Runtime/InteractiveToolsFramework/Private/InteractiveToolsContext.cpp
--- a/Engine/Source/Runtime/InteractiveToolsFramework/Private/InteractiveToolsContext.cpp
+++ b/Engine/Source/Runtime/InteractiveToolsFramework/Private/InteractiveToolsContext.cpp
@@ -166,10 +166,30 @@ void UInteractiveToolsContext::DeactivateAllActiveTools(EToolShutdownType Shutdo
 	auto DeactivateTool = [this, ShutdownType](EToolSide WhichSide) {
 		if (ToolManager->HasActiveTool(WhichSide))
 		{
-			const EToolShutdownType AcceptOrCancel =
-				ShutdownType != EToolShutdownType::Cancel && ToolManager->CanAcceptActiveTool(WhichSide)
-				? EToolShutdownType::Accept : EToolShutdownType::Cancel;
-			ToolManager->DeactivateTool(WhichSide, AcceptOrCancel);
+			EToolShutdownType ShutdownTypeToUse = ShutdownType;
+
+			// We do not allow passing an "accept" shutdown type if the tool says that it cannot accept.
+			// For complete-style tools, which can never accept, it should theoretically not matter what
+			// we pass because we would like them to respond the same way to all three shutdown types.
+			// However, we currently have one complete-style tool that responds differently to cancel, and we
+			// can't change that in a hotfix. So, we pass "complete" to complete-style tools.
+			if (ShutdownType != EToolShutdownType::Cancel)
+			{
+				if (!ToolManager->GetActiveTool(WhichSide)->HasAccept())
+				{
+					ShutdownTypeToUse = EToolShutdownType::Completed;
+				}
+				else if (ToolManager->CanAcceptActiveTool(WhichSide))
+				{
+					ShutdownTypeToUse = EToolShutdownType::Accept;
+				}
+				else
+				{
+					ShutdownTypeToUse = EToolShutdownType::Cancel;
+				}
+			}
+
+			ToolManager->DeactivateTool(WhichSide, ShutdownTypeToUse);
 		}
 	};
 	
