Add view alignment mode PointScreen to gizmo element base class.

diff --git a/Engine/Source/Editor/Experimental/EditorInteractiveToolsFramework/Private/EditorGizmos/TransformGizmo.cpp b/Engine/Source/Editor/Experimental/EditorInteractiveToolsFramework/Private/EditorGizmos/TransformGizmo.cpp
--- a/Engine/Source/Editor/Experimental/EditorInteractiveToolsFramework/Private/EditorGizmos/TransformGizmo.cpp
+++ b/Engine/Source/Editor/Experimental/EditorInteractiveToolsFramework/Private/EditorGizmos/TransformGizmo.cpp
@@ -889,7 +889,9 @@ UGizmoElementRectangle* UTransformGizmo::MakeTranslateScreenSpaceHandle()
 	RectangleElement->SetCenter(FVector::ZeroVector);
 	RectangleElement->SetHeight(TranslateScreenSpaceHandleSize);
 	RectangleElement->SetWidth(TranslateScreenSpaceHandleSize);
-	RectangleElement->SetScreenSpace(true);
+	RectangleElement->SetViewAlignType(EGizmoElementViewAlignType::PointScreen);
+	RectangleElement->SetViewAlignAxis(FVector::UpVector);
+	RectangleElement->SetViewAlignNormal(-FVector::ForwardVector);
 	RectangleElement->SetMaterial(TransparentVertexColorMaterial);
 	RectangleElement->SetLineColor(ScreenSpaceColor);
 	RectangleElement->SetHitMesh(true);
diff --git a/Engine/Source/Runtime/InteractiveToolsFramework/Private/BaseGizmos/GizmoElementBase.cpp b/Engine/Source/Runtime/InteractiveToolsFramework/Private/BaseGizmos/GizmoElementBase.cpp
--- a/Engine/Source/Runtime/InteractiveToolsFramework/Private/BaseGizmos/GizmoElementBase.cpp
+++ b/Engine/Source/Runtime/InteractiveToolsFramework/Private/BaseGizmos/GizmoElementBase.cpp
@@ -68,7 +68,7 @@ bool UGizmoElementBase::GetViewAlignRot(const FSceneView* View, const FTransform
 
 	FTransform WorldToLocalTransform = InLocalToWorldTransform.Inverse();
 
-	if (View->IsPerspectiveProjection())
+	if (View->IsPerspectiveProjection() && ViewAlignType != EGizmoElementViewAlignType::PointScreen)
 	{
 		FVector LocalViewLocation = WorldToLocalTransform.TransformPosition(View->ViewLocation);
 		LocalViewDir = InLocalCenter - LocalViewLocation;
@@ -85,33 +85,19 @@ bool UGizmoElementBase::GetViewAlignRot(const FSceneView* View, const FTransform
 	{
 		OutAlignRot = FQuat::FindBetweenNormals(ViewAlignNormal, -LocalViewDir);
 	}
-	else if (ViewAlignType == EGizmoElementViewAlignType::PointEye)
+	else if (ViewAlignType == EGizmoElementViewAlignType::PointEye || ViewAlignType == EGizmoElementViewAlignType::PointScreen)
 	{
 		FVector Right = ViewAlignAxis ^ ViewAlignNormal;
 		Right.Normalize();
 		FVector Up = ViewAlignNormal ^ Right;
 
-		FMatrix LocalToCanonicalBasis(
-			FPlane(ViewAlignNormal.X, Right.X, Up.X, 0.0),
-			FPlane(ViewAlignNormal.Y, Right.Y, Up.Y, 0.0),
-			FPlane(ViewAlignNormal.Z, Right.Z, Up.Z, 0.0),
-			FPlane::ZeroVector);
-
 		FVector LocalViewUp = WorldToLocalTransform.TransformVector(View->GetViewUp());
 		FVector TargetFwd = -LocalViewDir;
 		FVector TargetRight = LocalViewUp ^ TargetFwd;
 		TargetRight.Normalize();
 		FVector TargetUp = TargetFwd ^ TargetRight;
 
-		FMatrix CanonicalToAlignRotBasis(
-			FPlane(TargetFwd.X, TargetFwd.Y, TargetFwd.Z, 0.0),
-			FPlane(TargetRight.X, TargetRight.Y, TargetRight.Z, 0.0),
-			FPlane(TargetUp.X, TargetUp.Y, TargetUp.Z, 0.0),
-			FPlane::ZeroVector);
-
-		FMatrix LocalToAligned = LocalToCanonicalBasis * CanonicalToAlignRotBasis;
-		OutAlignRot = LocalToAligned.ToQuat();
-		OutAlignRot.Normalize();
+		OutAlignRot = GetAlignRotBetweenCoordSpaces(ViewAlignNormal, Right, Up, TargetFwd, TargetRight, TargetUp);
 	}
 	else if (ViewAlignType == EGizmoElementViewAlignType::Axial)
 	{
@@ -131,6 +117,26 @@ bool UGizmoElementBase::GetViewAlignRot(const FSceneView* View, const FTransform
 	return true;
 }
 
+FQuat UGizmoElementBase::GetAlignRotBetweenCoordSpaces(FVector SourceForward, FVector SourceRight, FVector SourceUp, FVector TargetForward, FVector TargetRight, FVector TargetUp) const
+{
+	FMatrix SourceToCanonical(
+		FPlane(SourceForward.X, SourceRight.X, SourceUp.X, 0.0),
+		FPlane(SourceForward.Y, SourceRight.Y, SourceUp.Y, 0.0),
+		FPlane(SourceForward.Z, SourceRight.Z, SourceUp.Z, 0.0),
+		FPlane::ZeroVector);
+
+	FMatrix CanonicalToTarget(
+		FPlane(TargetForward.X, TargetForward.Y, TargetForward.Z, 0.0),
+		FPlane(TargetRight.X, TargetRight.Y, TargetRight.Z, 0.0),
+		FPlane(TargetUp.X, TargetUp.Y, TargetUp.Z, 0.0),
+		FPlane::ZeroVector);
+
+	FMatrix SourceToTarget = SourceToCanonical * CanonicalToTarget;
+	FQuat Result = SourceToTarget.ToQuat();
+	Result.Normalize();
+	return Result;
+}
+
 const UMaterialInterface* UGizmoElementBase::GetCurrentMaterial(const FRenderTraversalState& RenderState) const
 {
 	EGizmoElementInteractionState CurrentState;
diff --git a/Engine/Source/Runtime/InteractiveToolsFramework/Private/BaseGizmos/GizmoElementCircle.cpp b/Engine/Source/Runtime/InteractiveToolsFramework/Private/BaseGizmos/GizmoElementCircle.cpp
--- a/Engine/Source/Runtime/InteractiveToolsFramework/Private/BaseGizmos/GizmoElementCircle.cpp
+++ b/Engine/Source/Runtime/InteractiveToolsFramework/Private/BaseGizmos/GizmoElementCircle.cpp
@@ -23,32 +23,23 @@ void UGizmoElementCircle::Render(IToolsContextRenderAPI* RenderAPI, const FRende
 
 	if (bVisibleViewDependent)
 	{
-		FVector Axis0, Axis1;
-
-		if (bScreenSpace)
+		FVector AdjustedNormal;
+		FQuat AlignRot;
+		if (GetViewAlignRot(View, LocalToWorldTransform, Center, AlignRot))
 		{
-			Axis0 = View->GetViewUp();
-			Axis1 = View->GetViewRight();
+			AdjustedNormal = AlignRot.RotateVector(Normal);
 		}
 		else
 		{
-			FVector AdjustedNormal;
-			FQuat AlignRot;
-			if (GetViewAlignRot(View, LocalToWorldTransform, Center, AlignRot))
-			{
-				AdjustedNormal = AlignRot.RotateVector(Normal);
-			}
-			else
-			{
-				AdjustedNormal = Normal;
-			}
-
-			const FVector WorldNormal = LocalToWorldTransform.TransformVectorNoScale(AdjustedNormal);
-			GizmoMath::MakeNormalPlaneBasis(WorldNormal, Axis0, Axis1);
-			Axis0.Normalize();
-			Axis1.Normalize();
+			AdjustedNormal = Normal;
 		}
 
+		FVector Axis0, Axis1;
+		const FVector WorldNormal = LocalToWorldTransform.TransformVectorNoScale(AdjustedNormal);
+		GizmoMath::MakeNormalPlaneBasis(WorldNormal, Axis0, Axis1);
+		Axis0.Normalize();
+		Axis1.Normalize();
+
 		const float WorldRadius = Radius * LocalToWorldTransform.GetScale3D().X;
 		const FVector WorldCenter = LocalToWorldTransform.TransformPosition(Center);
 
@@ -200,12 +191,3 @@ bool UGizmoElementCircle::GetHitLine() const
 	return bHitLine;
 }
 
-void UGizmoElementCircle::SetScreenSpace(bool InScreenSpace)
-{
-	bScreenSpace = InScreenSpace;
-}
-
-bool UGizmoElementCircle::GetScreenSpace()
-{
-	return bScreenSpace;
-}
diff --git a/Engine/Source/Runtime/InteractiveToolsFramework/Private/BaseGizmos/GizmoElementRectangle.cpp b/Engine/Source/Runtime/InteractiveToolsFramework/Private/BaseGizmos/GizmoElementRectangle.cpp
--- a/Engine/Source/Runtime/InteractiveToolsFramework/Private/BaseGizmos/GizmoElementRectangle.cpp
+++ b/Engine/Source/Runtime/InteractiveToolsFramework/Private/BaseGizmos/GizmoElementRectangle.cpp
@@ -26,29 +26,21 @@ void UGizmoElementRectangle::Render(IToolsContextRenderAPI* RenderAPI, const FRe
 	if (bVisibleViewDependent)
 	{
 		FVector Axis0, Axis1;
-		if (bScreenSpace)
+		FQuat AlignRot;
+		if (GetViewAlignRot(View, LocalToWorldTransform, Center, AlignRot))
 		{
-			Axis0 = View->GetViewUp();
-			Axis1 = View->GetViewRight();
+			Axis0 = AlignRot.RotateVector(UpDirection);
+			Axis1 = AlignRot.RotateVector(SideDirection);
 		}
 		else
 		{
-			FQuat AlignRot;
-			if (GetViewAlignRot(View, LocalToWorldTransform, Center, AlignRot))
-			{
-				Axis0 = AlignRot.RotateVector(UpDirection);
-				Axis1 = AlignRot.RotateVector(SideDirection);
-			}
-			else
-			{
-				Axis0 = UpDirection;
-				Axis1 = SideDirection;
-			}
-
-			Axis0 = LocalToWorldTransform.TransformVectorNoScale(Axis0);
-			Axis1 = LocalToWorldTransform.TransformVectorNoScale(Axis1);
+			Axis0 = UpDirection;
+			Axis1 = SideDirection;
 		}
 
+		Axis0 = LocalToWorldTransform.TransformVectorNoScale(Axis0);
+		Axis1 = LocalToWorldTransform.TransformVectorNoScale(Axis1);
+
 		FVector WorldCenter = LocalToWorldTransform.TransformPosition(Center);
 		float WorldWidth = Width * LocalToWorldTransform.GetScale3D().X;
 		float WorldHeight = Height * LocalToWorldTransform.GetScale3D().X;
@@ -178,16 +170,6 @@ FVector UGizmoElementRectangle::GetSideDirection() const
 	return SideDirection;
 }
 
-void UGizmoElementRectangle::SetScreenSpace(bool InScreenSpace)
-{
-	bScreenSpace = InScreenSpace;
-}
-
-bool UGizmoElementRectangle::GetScreenSpace() const
-{
-	return bScreenSpace;
-}
-
 void UGizmoElementRectangle::SetLineColor(const FColor& InLineColor)
 {
 	LineColor = InLineColor;
diff --git a/Engine/Source/Runtime/InteractiveToolsFramework/Public/BaseGizmos/GizmoElementBase.h b/Engine/Source/Runtime/InteractiveToolsFramework/Public/BaseGizmos/GizmoElementBase.h
--- a/Engine/Source/Runtime/InteractiveToolsFramework/Public/BaseGizmos/GizmoElementBase.h
+++ b/Engine/Source/Runtime/InteractiveToolsFramework/Public/BaseGizmos/GizmoElementBase.h
@@ -43,8 +43,9 @@ UENUM()
 enum class EGizmoElementViewAlignType
 {
 	None,
-	PointEye, // Align object forward axis to view direction, align object up axis to view up
 	PointOnly, // Align object forward axis to view direction only, useful for symmetrical objects such as a circle
+	PointEye, // Align object forward axis to -camera view direction (camera pos - object center), align object up axis to scene view up
+	PointScreen, // Align object forward axis to scene view forward direction (view up ^ view right), align object up axis to scene view up
 	Axial // Rotate object around up axis, minimizing angle between forward axis and view direction
 };
 
@@ -108,13 +109,6 @@ public:
 	// Calcute box sphere bounds for use when hit testing.
 	virtual FBoxSphereBounds CalcBounds(const FTransform & LocalToWorld) const PURE_VIRTUAL(UGizmoElementBase::CalcBounds, return FBoxSphereBounds(););
 
-	// Returns whether object is visible based on view-dependent visibility settings. 
-	virtual bool GetViewDependentVisibility(const FSceneView* View, const FTransform& InLocalToWorldTransform, const FVector& InLocalCenter) const;
-
-	// Returns true when view alignment is enabled. OutAlignRot is the rotation in local space which will align the object to the view base
-	// on view-dependent alignment settings. So it should be prepended to the local-to-world transform.
-	virtual bool GetViewAlignRot(const FSceneView* View, const FTransform& InLocalToWorldTransform, const FVector& InLocalCenter, FQuat& OutAlignRot) const;
-
 	// Returns the current material for given object state. The materials in the render state will override those in the element.
 	virtual const UMaterialInterface* GetCurrentMaterial(const FRenderTraversalState& RenderState) const;
 
@@ -330,6 +324,17 @@ protected:
 	bool bHasCachedBoxSphereBounds;
 
 protected:
+
+	// Returns whether object is visible based on view-dependent visibility settings. 
+	virtual bool GetViewDependentVisibility(const FSceneView* View, const FTransform& InLocalToWorldTransform, const FVector& InLocalCenter) const;
+
+	// Returns true when view alignment is enabled. OutAlignRot is the rotation in local space which will align the object to the view base
+	// on view-dependent alignment settings. So it should be prepended to the local-to-world transform.
+	virtual bool GetViewAlignRot(const FSceneView* View, const FTransform& InLocalToWorldTransform, const FVector& InLocalCenter, FQuat& OutAlignRot) const;
+
+	// Helper method to calculate rotation between coord spaces.
+	FQuat GetAlignRotBetweenCoordSpaces(FVector SourceForward, FVector SourceSide, FVector SourceUp, FVector TargetForward, FVector TargetSide, FVector TargetUp) const;
+
 	// Cache render state during render traversal, to be used subsequently when line tracing.
 	virtual void CacheRenderState(const FTransform& InLocalToWorldState, double InPixelToWorldScale, bool InVisibleViewDependent = true);
 
diff --git a/Engine/Source/Runtime/InteractiveToolsFramework/Public/BaseGizmos/GizmoElementCircle.h b/Engine/Source/Runtime/InteractiveToolsFramework/Public/BaseGizmos/GizmoElementCircle.h
--- a/Engine/Source/Runtime/InteractiveToolsFramework/Public/BaseGizmos/GizmoElementCircle.h
+++ b/Engine/Source/Runtime/InteractiveToolsFramework/Public/BaseGizmos/GizmoElementCircle.h
@@ -60,10 +60,6 @@ public:
 	virtual void SetHitLine(bool InHitLine);
 	virtual bool GetHitLine() const;
 
-	// Screen space, when true orients the circle to the screen up and side vectors
-	virtual void SetScreenSpace(bool InScreenSpace);
-	virtual bool GetScreenSpace();
-
 protected:
 
 	// Circle center.
@@ -101,8 +97,4 @@ protected:
 	// Whether to perform hit test on line.
 	UPROPERTY()
 	bool bHitLine = false;
-
-	// Screen space flag
-	UPROPERTY()
-	bool bScreenSpace = false;
 };
\ No newline at end of file
diff --git a/Engine/Source/Runtime/InteractiveToolsFramework/Public/BaseGizmos/GizmoElementRectangle.h b/Engine/Source/Runtime/InteractiveToolsFramework/Public/BaseGizmos/GizmoElementRectangle.h
--- a/Engine/Source/Runtime/InteractiveToolsFramework/Public/BaseGizmos/GizmoElementRectangle.h
+++ b/Engine/Source/Runtime/InteractiveToolsFramework/Public/BaseGizmos/GizmoElementRectangle.h
@@ -48,10 +48,6 @@ public:
 	virtual void SetSideDirection(const FVector& InSideDirection);
 	virtual FVector GetSideDirection() const;
 
-	// Screen space, when true orients the rectangle to the screen up and side vectors
-	virtual void SetScreenSpace(bool InScreenSpace);
-	virtual bool GetScreenSpace() const;
-
 	// Line color
 	virtual void SetLineColor(const FColor& InLineColor);
 	virtual FColor GetLineColor() const;
@@ -94,10 +90,6 @@ protected:
 	UPROPERTY()
 	FVector SideDirection = FVector(0.0f, 1.0f, 0.0f);
 
-	// Screen space flag
-	UPROPERTY()
-	bool bScreenSpace = false;
-
 	// Line color when bShowOutline is false.
 	UPROPERTY()
 	FColor LineColor = FColor::White;
