Added support for render state overrides and line render state in the gizmo element library.

diff --git a/Engine/Source/Editor/Experimental/EditorInteractiveToolsFramework/Private/EditorGizmos/TransformGizmo.cpp b/Engine/Source/Editor/Experimental/EditorInteractiveToolsFramework/Private/EditorGizmos/TransformGizmo.cpp
--- a/Engine/Source/Editor/Experimental/EditorInteractiveToolsFramework/Private/EditorGizmos/TransformGizmo.cpp
+++ b/Engine/Source/Editor/Experimental/EditorInteractiveToolsFramework/Private/EditorGizmos/TransformGizmo.cpp
@@ -41,6 +41,8 @@ void UTransformGizmo::Setup()
 	GizmoElementRoot->SetConstantScale(true);
 	GizmoElementRoot->SetHoverMaterial(CurrentAxisMaterial);
 	GizmoElementRoot->SetInteractMaterial(CurrentAxisMaterial);
+	GizmoElementRoot->SetHoverLineColor(CurrentColor);
+	GizmoElementRoot->SetInteractLineColor(CurrentColor);
 
 	bInInteraction = false;
 }
@@ -858,8 +860,8 @@ UGizmoElementRectangle* UTransformGizmo::MakePlanarHandle(ETransformGizmoPartIde
 {
 	FVector PlanarHandleCenter = (InUpDirection + InSideDirection) * PlanarHandleOffset;
 
-	FColor LineColor = InVertexColor.ToFColor(false);
-	FColor VertexColor = LineColor;
+	FLinearColor LineColor = InVertexColor;
+	FLinearColor VertexColor = LineColor;
 	VertexColor.A = LargeOuterAlpha;
 
 	UGizmoElementRectangle* RectangleElement = NewObject<UGizmoElementRectangle>();
@@ -897,7 +899,8 @@ UGizmoElementRectangle* UTransformGizmo::MakeTranslateScreenSpaceHandle()
 	RectangleElement->SetHitMesh(true);
 	RectangleElement->SetDrawMesh(false);
 	RectangleElement->SetDrawLine(true);
-	RectangleElement->SetHoverLineThicknessMultiplier(3.0);
+	RectangleElement->SetHoverLineThicknessMultiplier(3.0f);
+	RectangleElement->SetInteractLineThicknessMultiplier(3.0f);
 	return RectangleElement;
 }
 
@@ -937,7 +940,7 @@ UGizmoElementCircle* UTransformGizmo::MakeRotateCircleHandle(ETransformGizmoPart
 
 	if (bFill)
 	{
-		CircleElement->SetVertexColor(InColor.ToFColor(true));
+		CircleElement->SetVertexColor(InColor);
 		CircleElement->SetMaterial(WhiteMaterial);
 	}
 	else
diff --git a/Engine/Source/Editor/Experimental/EditorInteractiveToolsFramework/Public/EditorGizmos/TransformGizmo.h b/Engine/Source/Editor/Experimental/EditorInteractiveToolsFramework/Public/EditorGizmos/TransformGizmo.h
--- a/Engine/Source/Editor/Experimental/EditorInteractiveToolsFramework/Public/EditorGizmos/TransformGizmo.h
+++ b/Engine/Source/Editor/Experimental/EditorInteractiveToolsFramework/Public/EditorGizmos/TransformGizmo.h
@@ -115,12 +115,11 @@ public:
 	static constexpr FLinearColor AxisColorX = FLinearColor(0.594f, 0.0197f, 0.0f);
 	static constexpr FLinearColor AxisColorY = FLinearColor(0.1349f, 0.3959f, 0.0f);
 	static constexpr FLinearColor AxisColorZ = FLinearColor(0.0251f, 0.207f, 0.85f);
-	static constexpr FLinearColor ScreenAxisColor = FLinearColor(0.76, 0.72, 0.14f);
-	static constexpr FColor PlaneColorXY = FColor(255, 255, 0); // FColor::Yellow
-	static constexpr FColor ArcBallColor = FColor(128, 128, 128, 6);
-	static constexpr FColor ScreenSpaceColor = FColor(196, 196, 196);
-	static constexpr FColor CurrentColor = FColor(255, 255, 0); // FColor::Yellow
-
+	static constexpr FLinearColor ScreenAxisColor = FLinearColor(0.76f, 0.72f, 0.14f);
+	static constexpr FLinearColor PlaneColorXY = FLinearColor(1.0f, 1.0f, 0.0f);
+	static constexpr FLinearColor ArcBallColor = FLinearColor(0.5f, 0.5f, 0.5f, 0.03f);
+	static constexpr FLinearColor ScreenSpaceColor = FLinearColor(0.765f, 0.765f, 0.765f);
+	static constexpr FLinearColor CurrentColor = FLinearColor(1.0f, 1.0f, 0.0f);
 	static constexpr FLinearColor GreyColor = FLinearColor(0.50f, 0.50f, 0.50f);
 	static constexpr FLinearColor WhiteColor = FLinearColor(1.0f, 1.0f, 1.0f);
 
@@ -128,10 +127,7 @@ public:
 	static constexpr FLinearColor RotateOuterCircleColor = GreyColor;
 	static constexpr FLinearColor RotateArcballCircleColor = WhiteColor;
 
-	static constexpr uint8 LargeInnerAlpha = 0x3f;
-	static constexpr uint8 SmallInnerAlpha = 0x0f;
-	static constexpr uint8 LargeOuterAlpha = 0x7f;
-	static constexpr uint8 SmallOuterAlpha = 0x0f;
+	static constexpr float LargeOuterAlpha = 0.5f;
 
 
 public:
diff --git a/Engine/Source/Runtime/InteractiveToolsFramework/Private/BaseGizmos/GizmoElementArrow.cpp b/Engine/Source/Runtime/InteractiveToolsFramework/Private/BaseGizmos/GizmoElementArrow.cpp
--- a/Engine/Source/Runtime/InteractiveToolsFramework/Private/BaseGizmos/GizmoElementArrow.cpp
+++ b/Engine/Source/Runtime/InteractiveToolsFramework/Private/BaseGizmos/GizmoElementArrow.cpp
@@ -24,6 +24,8 @@ void UGizmoElementArrow::Render(IToolsContextRenderAPI* RenderAPI, const FRender
 		return;
 	}
 
+	check(RenderAPI);
+
 	if (bUpdateArrowBody)
 	{
 		UpdateArrowBody();
@@ -34,42 +36,26 @@ void UGizmoElementArrow::Render(IToolsContextRenderAPI* RenderAPI, const FRender
 		UpdateArrowHead();
 	}
 
-	const FSceneView* View = RenderAPI->GetSceneView();
-	bool bVisibleViewDependent = GetViewDependentVisibility(View, RenderState.LocalToWorldTransform, Base);
-
-	FRenderTraversalState RenderStateCopy = RenderState;
+	FRenderTraversalState CurrentRenderState(RenderState);
+	bool bVisibleViewDependent = UpdateRenderState(RenderAPI, Base, CurrentRenderState);
 
 	if (bVisibleViewDependent)
 	{
-		FQuat AlignRot;
-		if (GetViewAlignRot(View, RenderState.LocalToWorldTransform, Base, AlignRot))
-		{
-			RenderStateCopy.LocalToWorldTransform = FTransform(AlignRot, Base) * RenderState.LocalToWorldTransform;
-		}
-		else
-		{
-			RenderStateCopy.LocalToWorldTransform = FTransform(Base) * RenderState.LocalToWorldTransform;
-		}
-
-		UpdateRenderTraversalState(RenderStateCopy);
-
 		check(CylinderElement);
 		UGizmoElementBase* Element = Cast<UGizmoElementBase>(CylinderElement);
-		Element->Render(RenderAPI, RenderStateCopy);
+		Element->Render(RenderAPI, CurrentRenderState);
 
 		if (HeadType == EGizmoElementArrowHeadType::Cone)
 		{
 			check(ConeElement);
-			ConeElement->Render(RenderAPI, RenderStateCopy);
+			ConeElement->Render(RenderAPI, CurrentRenderState);
 		}
 		else // (HeadType == EGizmoElementArrowHeadType::Cube)
 		{
 			check(BoxElement);
-			BoxElement->Render(RenderAPI, RenderStateCopy);
+			BoxElement->Render(RenderAPI, CurrentRenderState);
 		}
 	}
-
-	CacheRenderState(RenderStateCopy.LocalToWorldTransform, RenderStateCopy.PixelToWorldScale, bVisibleViewDependent);
 }
 
 FInputRayHit UGizmoElementArrow::LineTrace(const FVector RayOrigin, const FVector RayDirection)
diff --git a/Engine/Source/Runtime/InteractiveToolsFramework/Private/BaseGizmos/GizmoElementBase.cpp b/Engine/Source/Runtime/InteractiveToolsFramework/Private/BaseGizmos/GizmoElementBase.cpp
--- a/Engine/Source/Runtime/InteractiveToolsFramework/Private/BaseGizmos/GizmoElementBase.cpp
+++ b/Engine/Source/Runtime/InteractiveToolsFramework/Private/BaseGizmos/GizmoElementBase.cpp
@@ -137,45 +137,6 @@ FQuat UGizmoElementBase::GetAlignRotBetweenCoordSpaces(FVector SourceForward, FV
 	return Result;
 }
 
-const UMaterialInterface* UGizmoElementBase::GetCurrentMaterial(const FRenderTraversalState& RenderState) const
-{
-	EGizmoElementInteractionState CurrentState;
-
-	if (RenderState.InteractionState == EGizmoElementInteractionState::None)
-	{
-		CurrentState = ElementInteractionState;
-	}
-	else
-	{
-		CurrentState = RenderState.InteractionState;
-	}
-
-	if (CurrentState == EGizmoElementInteractionState::Hovering)
-	{
-		if (RenderState.HoverMaterial.IsValid())
-		{
-			return RenderState.HoverMaterial.Get();
-		}
-		return HoverMaterial;
-	}
-
-	if (CurrentState == EGizmoElementInteractionState::Interacting)
-	{
-		if (RenderState.InteractMaterial.IsValid())
-		{
-			return RenderState.InteractMaterial.Get();
-		}
-		return InteractMaterial;
-	}
-
-	// CurrentState is None, so just return the regular material
-	if (RenderState.Material.IsValid())
-	{
-		return RenderState.Material.Get();
-	}
-	return Material;
-}
-
 void UGizmoElementBase::CacheRenderState(const FTransform& InLocalToWorldState, double InPixelToWorldScale, bool InVisibleViewDependent)
 {
 	CachedLocalToWorldTransform = InLocalToWorldState;
@@ -192,24 +153,37 @@ void UGizmoElementBase::ResetCachedRenderState()
 	bCachedVisibleViewDependent = true;
 }
 
-void UGizmoElementBase::UpdateRenderTraversalState(FRenderTraversalState& InRenderTraversalState)
+bool UGizmoElementBase::UpdateRenderState(IToolsContextRenderAPI* RenderAPI, const FVector& InLocalOrigin, FRenderTraversalState& InOutRenderState)
 {
-	if (InRenderTraversalState.InteractionState == EGizmoElementInteractionState::None)
-	{
-		InRenderTraversalState.InteractionState = ElementInteractionState;
-	}
-	if (!InRenderTraversalState.HoverMaterial.IsValid())
-	{
-		InRenderTraversalState.HoverMaterial = HoverMaterial;
-	}
-	if (!InRenderTraversalState.InteractMaterial.IsValid())
+	check(RenderAPI);
+	const FSceneView* View = RenderAPI->GetSceneView();
+	check(View);
+
+	if (InOutRenderState.InteractionState == EGizmoElementInteractionState::None)
 	{
-		InRenderTraversalState.InteractMaterial = InteractMaterial;
+		InOutRenderState.InteractionState = ElementInteractionState;
 	}
-	if (!InRenderTraversalState.Material.IsValid())
+
+	InOutRenderState.MeshRenderState.Update(MeshRenderAttributes);
+
+	bool bVisibleViewDependent = GetViewDependentVisibility(View, InOutRenderState.LocalToWorldTransform, InLocalOrigin);
+
+	if (bVisibleViewDependent)
 	{
-		InRenderTraversalState.Material = Material;
+		FQuat AlignRot;
+		if (GetViewAlignRot(View, InOutRenderState.LocalToWorldTransform, InLocalOrigin, AlignRot))
+		{
+			InOutRenderState.LocalToWorldTransform = FTransform(AlignRot, InLocalOrigin) * InOutRenderState.LocalToWorldTransform;
+		}
+		else
+		{
+			InOutRenderState.LocalToWorldTransform = FTransform(InLocalOrigin) * InOutRenderState.LocalToWorldTransform;
+		}
 	}
+
+	CacheRenderState(InOutRenderState.LocalToWorldTransform, InOutRenderState.PixelToWorldScale, bVisibleViewDependent);
+
+	return bVisibleViewDependent;
 }
 
 bool UGizmoElementBase::IsVisible() const
@@ -374,53 +348,97 @@ float UGizmoElementBase::GetViewAlignAxialAngleTol() const
 	return ViewAlignAxialAngleTol;
 }
 
-void UGizmoElementBase::SetMaterial(UMaterialInterface* InMaterial)
+void UGizmoElementBase::SetPixelHitDistanceThreshold(float InPixelHitDistanceThreshold)
 {
-	Material = InMaterial;
+	PixelHitDistanceThreshold = InPixelHitDistanceThreshold;
 }
 
-UMaterialInterface* UGizmoElementBase::GetMaterial() const
+float UGizmoElementBase::GetPixelHitDistanceThreshold() const
 {
-	return Material;
+	return PixelHitDistanceThreshold;
 }
 
-void UGizmoElementBase::SetHoverMaterial(UMaterialInterface* InHoverMaterial)
+void UGizmoElementBase::SetMaterial(TWeakObjectPtr<UMaterialInterface> InMaterial, bool InOverridesChildState)
 {
-	HoverMaterial = InHoverMaterial;
+	MeshRenderAttributes.Material.SetMaterial(InMaterial, InOverridesChildState);
 }
 
-UMaterialInterface* UGizmoElementBase::GetHoverMaterial() const
+const UMaterialInterface* UGizmoElementBase::GetMaterial() const
 {
-	return HoverMaterial;
+	return MeshRenderAttributes.Material.GetMaterial();
 }
 
-void UGizmoElementBase::SetInteractMaterial(UMaterialInterface* InInteractMaterial)
+bool UGizmoElementBase::GetMaterialOverridesChildState() const
 {
-	InteractMaterial = InInteractMaterial;
+	return MeshRenderAttributes.Material.bOverridesChildState;
 }
 
-UMaterialInterface* UGizmoElementBase::GetInteractMaterial() const
+void UGizmoElementBase::ClearMaterial()
 {
-	return InteractMaterial;
+	MeshRenderAttributes.Material.Reset();
 }
 
-void UGizmoElementBase::SetVertexColor(const FColor& InVertexColor)
+void UGizmoElementBase::SetHoverMaterial(TWeakObjectPtr<UMaterialInterface> InMaterial, bool InOverridesChildState)
 {
-	VertexColor = InVertexColor;
+	MeshRenderAttributes.HoverMaterial.SetMaterial(InMaterial, InOverridesChildState);
 }
 
-FColor UGizmoElementBase::GetVertexColor() const
+const UMaterialInterface* UGizmoElementBase::GetHoverMaterial() const
 {
-	return VertexColor;
+	return MeshRenderAttributes.HoverMaterial.GetMaterial();
 }
 
-void UGizmoElementBase::SetPixelHitDistanceThreshold(float InPixelHitDistanceThreshold)
+bool UGizmoElementBase::GetHoverMaterialOverridesChildState() const
 {
-	PixelHitDistanceThreshold = InPixelHitDistanceThreshold;
+	return MeshRenderAttributes.HoverMaterial.bOverridesChildState;
 }
 
-float UGizmoElementBase::GetPixelHitDistanceThreshold() const
+void UGizmoElementBase::ClearHoverMaterial()
 {
-	return PixelHitDistanceThreshold;
+	MeshRenderAttributes.HoverMaterial.Reset();
+}
+
+void UGizmoElementBase::SetInteractMaterial(TWeakObjectPtr<UMaterialInterface> InMaterial, bool InOverridesChildState)
+{
+	MeshRenderAttributes.InteractMaterial.SetMaterial(InMaterial, InOverridesChildState);
+}
+
+const UMaterialInterface* UGizmoElementBase::GetInteractMaterial() const
+{
+	return MeshRenderAttributes.InteractMaterial.GetMaterial();
+}
+
+bool UGizmoElementBase::GetInteractMaterialOverridesChildState() const
+{
+	return MeshRenderAttributes.InteractMaterial.bOverridesChildState;
+}
+
+void UGizmoElementBase::ClearInteractMaterial()
+{
+	MeshRenderAttributes.InteractMaterial.Reset();
+}
+
+void UGizmoElementBase::SetVertexColor(FLinearColor InVertexColor, bool InOverridesChildState)
+{
+	MeshRenderAttributes.VertexColor.SetColor(InVertexColor, InOverridesChildState);
+}
+
+FLinearColor UGizmoElementBase::GetVertexColor() const
+{
+	return MeshRenderAttributes.VertexColor.GetColor();
+}
+bool UGizmoElementBase::HasVertexColor() const
+{
+	return MeshRenderAttributes.VertexColor.bHasValue;
+}
+
+bool UGizmoElementBase::GetVertexColorOverridesChildState() const
+{
+	return MeshRenderAttributes.VertexColor.bOverridesChildState;
+}
+
+void UGizmoElementBase::ClearVertexColor()
+{
+	MeshRenderAttributes.VertexColor.Reset();
 }
 
diff --git a/Engine/Source/Runtime/InteractiveToolsFramework/Private/BaseGizmos/GizmoElementBox.cpp b/Engine/Source/Runtime/InteractiveToolsFramework/Private/BaseGizmos/GizmoElementBox.cpp
--- a/Engine/Source/Runtime/InteractiveToolsFramework/Private/BaseGizmos/GizmoElementBox.cpp
+++ b/Engine/Source/Runtime/InteractiveToolsFramework/Private/BaseGizmos/GizmoElementBox.cpp
@@ -15,40 +15,21 @@ void UGizmoElementBox::Render(IToolsContextRenderAPI* RenderAPI, const FRenderTr
 	}
 
 	check(RenderAPI);
-	const FSceneView* View = RenderAPI->GetSceneView();
 
-	FTransform LocalToWorldTransform = RenderState.LocalToWorldTransform;
-
-	bool bVisibleViewDependent = GetViewDependentVisibility(View, LocalToWorldTransform, Center);
+	FRenderTraversalState CurrentRenderState(RenderState);
+	bool bVisibleViewDependent = UpdateRenderState(RenderAPI, Center, CurrentRenderState);
 
 	if (bVisibleViewDependent)
 	{
-		const UMaterialInterface* UseMaterial = GetCurrentMaterial(RenderState);
-
-		if (UseMaterial)
+		if (const UMaterialInterface* UseMaterial = CurrentRenderState.GetCurrentMaterial())
 		{
-			FQuat AlignRot;
-			FVector AdjustedSideDir, AdjustedUpDir;
-			if (GetViewAlignRot(View, LocalToWorldTransform, Center, AlignRot))
-			{
-				AdjustedSideDir = AlignRot.RotateVector(SideDirection);
-				AdjustedUpDir = AlignRot.RotateVector(UpDirection);
-			}
-			else
-			{
-				AdjustedSideDir = SideDirection;
-				AdjustedUpDir = UpDirection;
-			}
-
-			FQuat Rotation = FRotationMatrix::MakeFromYZ(AdjustedSideDir, AdjustedUpDir).ToQuat();
-			FPrimitiveDrawInterface* PDI = RenderAPI->GetPrimitiveDrawInterface();
-			FTransform DrawTransform = FTransform(Rotation, Center) * LocalToWorldTransform;
+			FQuat LocalRotation = FRotationMatrix::MakeFromYZ(SideDirection, UpDirection).ToQuat();
+			FTransform RenderLocalToWorldTransform = FTransform(LocalRotation) * CurrentRenderState.LocalToWorldTransform;
 			const FVector HalfDimensions = Dimensions * 0.5;
-			DrawBox(PDI, DrawTransform.ToMatrixWithScale(), HalfDimensions, UseMaterial->GetRenderProxy(), SDPG_Foreground);
+			FPrimitiveDrawInterface* PDI = RenderAPI->GetPrimitiveDrawInterface();
+			DrawBox(PDI, RenderLocalToWorldTransform.ToMatrixWithScale(), HalfDimensions, UseMaterial->GetRenderProxy(), SDPG_Foreground);
 		}
 	}
-
-	CacheRenderState(LocalToWorldTransform, RenderState.PixelToWorldScale, bVisibleViewDependent);
 }
 
 FInputRayHit UGizmoElementBox::LineTrace(const FVector RayOrigin, const FVector RayDirection)
@@ -58,7 +39,7 @@ FInputRayHit UGizmoElementBox::LineTrace(const FVector RayOrigin, const FVector
 		const FVector YAxis = CachedLocalToWorldTransform.TransformVectorNoScale(SideDirection);
 		const FVector ZAxis = CachedLocalToWorldTransform.TransformVectorNoScale(UpDirection);
 		const FVector XAxis = FVector::CrossProduct(YAxis, ZAxis);
-		const FVector WorldCenter = CachedLocalToWorldTransform.TransformPosition(Center);
+		const FVector WorldCenter = CachedLocalToWorldTransform.TransformPosition(FVector::ZeroVector);
 		const double Scale = CachedLocalToWorldTransform.GetScale3D().X;
 		const FVector WorldExtent = Dimensions * Scale * 0.5;
 
diff --git a/Engine/Source/Runtime/InteractiveToolsFramework/Private/BaseGizmos/GizmoElementCircle.cpp b/Engine/Source/Runtime/InteractiveToolsFramework/Private/BaseGizmos/GizmoElementCircle.cpp
--- a/Engine/Source/Runtime/InteractiveToolsFramework/Private/BaseGizmos/GizmoElementCircle.cpp
+++ b/Engine/Source/Runtime/InteractiveToolsFramework/Private/BaseGizmos/GizmoElementCircle.cpp
@@ -16,68 +16,60 @@ void UGizmoElementCircle::Render(IToolsContextRenderAPI* RenderAPI, const FRende
 	}
 
 	check(RenderAPI);
-	const FSceneView* View = RenderAPI->GetSceneView();
 
-	FTransform LocalToWorldTransform = RenderState.LocalToWorldTransform;
-	bool bVisibleViewDependent = GetViewDependentVisibility(View, LocalToWorldTransform, Center);
+	FRenderTraversalState CurrentRenderState(RenderState);
+	bool bVisibleViewDependent = UpdateRenderState(RenderAPI, Center, CurrentRenderState);
 
 	if (bVisibleViewDependent)
 	{
-		FVector AdjustedNormal;
-		FQuat AlignRot;
-		if (GetViewAlignRot(View, LocalToWorldTransform, Center, AlignRot))
-		{
-			AdjustedNormal = AlignRot.RotateVector(Normal);
-		}
-		else
-		{
-			AdjustedNormal = Normal;
-		}
-
-		FVector Axis0, Axis1;
-		const FVector WorldNormal = LocalToWorldTransform.TransformVectorNoScale(AdjustedNormal);
-		GizmoMath::MakeNormalPlaneBasis(WorldNormal, Axis0, Axis1);
-		Axis0.Normalize();
-		Axis1.Normalize();
-
-		const float WorldRadius = Radius * LocalToWorldTransform.GetScale3D().X;
-		const FVector WorldCenter = LocalToWorldTransform.TransformPosition(Center);
+		const float WorldRadius = Radius * CurrentRenderState.LocalToWorldTransform.GetScale3D().X;
+		const FVector WorldCenter = CurrentRenderState.LocalToWorldTransform.TransformPosition(FVector::ZeroVector);
+		FVector WorldUpAxis, WorldSideAxis;
+		const FVector WorldNormal = CurrentRenderState.LocalToWorldTransform.TransformVectorNoScale(Normal);
+		GizmoMath::MakeNormalPlaneBasis(WorldNormal, WorldUpAxis, WorldSideAxis);
+		WorldUpAxis.Normalize();
+		WorldSideAxis.Normalize();
 
 		FPrimitiveDrawInterface* PDI = RenderAPI->GetPrimitiveDrawInterface();
 
 		if (bDrawMesh)
 		{
-			if (const UMaterialInterface* UseMaterial = GetCurrentMaterial(RenderState))
+			if (const UMaterialInterface* UseMaterial = CurrentRenderState.GetCurrentMaterial())
 			{
-				DrawDisc(PDI, WorldCenter, Axis0, Axis1, VertexColor, WorldRadius, NumSides, UseMaterial->GetRenderProxy(), SDPG_Foreground);
+				FColor VertexColor = CurrentRenderState.GetVertexColor().ToFColor(false);
+				DrawDisc(PDI, WorldCenter, WorldUpAxis, WorldSideAxis, VertexColor, WorldRadius, NumSides, UseMaterial->GetRenderProxy(), SDPG_Foreground);
 			}
 		}
 
 		if (bDrawLine)
 		{
-			DrawCircle(PDI, WorldCenter, Axis0, Axis1, LineColor, WorldRadius, NumSides, SDPG_Foreground, GetCurrentLineThickness());
+			DrawCircle(PDI, WorldCenter, WorldUpAxis, WorldSideAxis, CurrentRenderState.GetCurrentLineColor(), WorldRadius, NumSides, SDPG_Foreground, GetCurrentLineThickness());
 		}
 	}
-	CacheRenderState(LocalToWorldTransform, RenderState.PixelToWorldScale, bVisibleViewDependent);
 }
 
 FInputRayHit UGizmoElementCircle::LineTrace(const FVector RayOrigin, const FVector RayDirection)
 {
 	if (IsHittableInView())
 	{
-		if (bHitMesh)
+		const FVector WorldNormal = CachedLocalToWorldTransform.TransformVectorNoScale(Normal);
+		const FVector WorldCenter = CachedLocalToWorldTransform.TransformPosition(FVector::ZeroVector);
+		const double PixelHitThresholdAdjust = CachedPixelToWorldScale * PixelHitDistanceThreshold;
+		double WorldRadius = CachedLocalToWorldTransform.GetScale3D().X * Radius;
+
+		// if ray is parallel to circle, no hit
+		if (FMath::IsNearlyZero(FVector::DotProduct(WorldNormal, RayDirection)))
 		{
-			bool bIntersects = false;
-			double Param = 0.0;
+			return FInputRayHit();
+		}
 
-			const FVector WorldNormal = CachedLocalToWorldTransform.TransformVectorNoScale(Normal);
-			const FVector WorldCenter = CachedLocalToWorldTransform.TransformPosition(Center);
-			const double PixelHitThresholdAdjust = CachedPixelToWorldScale * PixelHitDistanceThreshold;
-			const double WorldRadius = CachedLocalToWorldTransform.GetScale3D().X * Radius + PixelHitThresholdAdjust;
+		if (bHitMesh)
+		{
+			WorldRadius += PixelHitThresholdAdjust;
 
 			UE::Geometry::FLinearIntersection Result;
 			IntersectionUtil::RayCircleIntersection(RayOrigin, RayDirection, WorldCenter, WorldRadius, WorldNormal, Result);
-			
+
 			if (Result.intersects)
 			{
 				FInputRayHit RayHit(Result.parameter.Min);
@@ -141,16 +133,6 @@ int UGizmoElementCircle::GetNumSides() const
 	return NumSides;
 }
 
-void UGizmoElementCircle::SetLineColor(const FLinearColor& InColor)
-{
-	LineColor = InColor;
-}
-
-FLinearColor UGizmoElementCircle::GetLineColor() const
-{
-	return LineColor;
-}
-
 void UGizmoElementCircle::SetDrawMesh(bool InDrawMesh)
 {
 	bDrawMesh = InDrawMesh;
diff --git a/Engine/Source/Runtime/InteractiveToolsFramework/Private/BaseGizmos/GizmoElementCone.cpp b/Engine/Source/Runtime/InteractiveToolsFramework/Private/BaseGizmos/GizmoElementCone.cpp
--- a/Engine/Source/Runtime/InteractiveToolsFramework/Private/BaseGizmos/GizmoElementCone.cpp
+++ b/Engine/Source/Runtime/InteractiveToolsFramework/Private/BaseGizmos/GizmoElementCone.cpp
@@ -14,39 +14,25 @@ void UGizmoElementCone::Render(IToolsContextRenderAPI* RenderAPI, const FRenderT
 		return;
 	}
 
-	check(RenderAPI);
-	const FSceneView* View = RenderAPI->GetSceneView();
-
-	FTransform LocalToWorldTransform = RenderState.LocalToWorldTransform;
-
-	bool bVisibleViewDependent = GetViewDependentVisibility(View, LocalToWorldTransform, Origin);
+	FRenderTraversalState CurrentRenderState(RenderState);
+	bool bVisibleViewDependent = UpdateRenderState(RenderAPI, Origin, CurrentRenderState);
 
 	if (bVisibleViewDependent)
 	{
-		const UMaterialInterface* UseMaterial = GetCurrentMaterial(RenderState);
-
-		if (UseMaterial)
+		if (const UMaterialInterface* UseMaterial = CurrentRenderState.GetCurrentMaterial())
 		{
-			FQuat AlignRot;
-			const bool bHasAlignRot = GetViewAlignRot(View, LocalToWorldTransform, Origin, AlignRot);
-			const FVector AdjustedDir = bHasAlignRot ? AlignRot.RotateVector(Direction) : Direction;
-
-			FQuat Rotation = FRotationMatrix::MakeFromX(AdjustedDir).ToQuat();
-
+			FQuat Rotation = FRotationMatrix::MakeFromX(Direction).ToQuat();
 			const FVector Scale(Height);
 
-			FTransform RenderLocalToWorldTransform = FTransform(Rotation, Origin, Scale) * LocalToWorldTransform;
+			FTransform RenderLocalToWorldTransform = FTransform(Rotation, FVector::ZeroVector, Scale) * CurrentRenderState.LocalToWorldTransform;
 			const float ConeSide = FMath::Sqrt(Height * Height + Radius * Radius);
 			const float Angle = FMath::Acos(Height / ConeSide);
 			FPrimitiveDrawInterface* PDI = RenderAPI->GetPrimitiveDrawInterface();
 			DrawCone(PDI, RenderLocalToWorldTransform.ToMatrixWithScale(), Angle, Angle, NumSides, false, FColor::White, UseMaterial->GetRenderProxy(), SDPG_Foreground);
 		}
 	}
-	CacheRenderState(LocalToWorldTransform, RenderState.PixelToWorldScale, bVisibleViewDependent);
 }
 
-
-
 FInputRayHit UGizmoElementCone::LineTrace(const FVector RayOrigin, const FVector RayDirection)
 {
 	if (IsHittableInView())
@@ -57,15 +43,15 @@ FInputRayHit UGizmoElementCone::LineTrace(const FVector RayOrigin, const FVector
 		const double PixelHitThresholdAdjust = CachedPixelToWorldScale * PixelHitDistanceThreshold;
 		const double ConeSide = FMath::Sqrt(Height * Height + Radius * Radius);
 		const double CosAngle = Height / ConeSide;
-		const double ConeHeight = Height * CachedLocalToWorldTransform.GetScale3D().X + PixelHitThresholdAdjust * 2.0;
-		const FVector ConeDirection = CachedLocalToWorldTransform.TransformVectorNoScale(Direction);
-		const FVector ConeOrigin = CachedLocalToWorldTransform.TransformPosition(Origin) - ConeDirection * PixelHitThresholdAdjust;
+		const double WorldHeight = Height * CachedLocalToWorldTransform.GetScale3D().X + PixelHitThresholdAdjust * 2.0;
+		const FVector WorldDirection = CachedLocalToWorldTransform.TransformVectorNoScale(Direction);
+		const FVector WorldOrigin = CachedLocalToWorldTransform.TransformPosition(FVector::ZeroVector) - WorldDirection * PixelHitThresholdAdjust;
 
 		GizmoMath::RayConeIntersection(
-			ConeOrigin,
-			ConeDirection,
+			WorldOrigin,
+			WorldDirection,
 			CosAngle,
-			ConeHeight,
+			WorldHeight,
 			RayOrigin, RayDirection,
 			bIntersects, RayParam);
 
diff --git a/Engine/Source/Runtime/InteractiveToolsFramework/Private/BaseGizmos/GizmoElementCylinder.cpp b/Engine/Source/Runtime/InteractiveToolsFramework/Private/BaseGizmos/GizmoElementCylinder.cpp
--- a/Engine/Source/Runtime/InteractiveToolsFramework/Private/BaseGizmos/GizmoElementCylinder.cpp
+++ b/Engine/Source/Runtime/InteractiveToolsFramework/Private/BaseGizmos/GizmoElementCylinder.cpp
@@ -15,33 +15,23 @@ void UGizmoElementCylinder::Render(IToolsContextRenderAPI* RenderAPI, const FRen
 	}
 
 	check(RenderAPI);
-	const FSceneView* View = RenderAPI->GetSceneView();
 
-	FTransform LocalToWorldTransform = RenderState.LocalToWorldTransform;
-
-	bool bVisibleViewDependent = GetViewDependentVisibility(View, LocalToWorldTransform, Base);
+	FRenderTraversalState CurrentRenderState(RenderState);
+	bool bVisibleViewDependent = UpdateRenderState(RenderAPI, Base, CurrentRenderState);
 
 	if (bVisibleViewDependent)
 	{
-		const UMaterialInterface* UseMaterial = GetCurrentMaterial(RenderState);
-
-		if (UseMaterial)
+		if (const UMaterialInterface* UseMaterial = CurrentRenderState.GetCurrentMaterial())
 		{
-			FQuat AlignRot;
-			const bool bHasAlignRot = GetViewAlignRot(View, LocalToWorldTransform, Base, AlignRot);
-			const FVector AdjustedDir = bHasAlignRot ? AlignRot.RotateVector(Direction) : Direction;
-
-			const FQuat Rotation = FRotationMatrix::MakeFromZ(AdjustedDir).ToQuat();
+			const FQuat Rotation = FRotationMatrix::MakeFromZ(Direction).ToQuat();
 			const double HalfHeight = Height * 0.5;
-			const FVector Origin = Base + AdjustedDir * HalfHeight;
+			const FVector OriginOffset = Direction * HalfHeight;
 
-			FTransform RenderLocalToWorldTransform = FTransform(Rotation, Origin) * LocalToWorldTransform;
+			FTransform RenderLocalToWorldTransform = FTransform(Rotation, OriginOffset) * CurrentRenderState.LocalToWorldTransform;
 			FPrimitiveDrawInterface* PDI = RenderAPI->GetPrimitiveDrawInterface();
 			DrawCylinder(PDI, RenderLocalToWorldTransform.ToMatrixWithScale(), FVector::ZeroVector, FVector(1, 0, 0), FVector(0, 1, 0), FVector(0, 0, 1), Radius, HalfHeight, NumSides, UseMaterial->GetRenderProxy(), SDPG_Foreground);
-
 		}
 	}
-	CacheRenderState(LocalToWorldTransform, RenderState.PixelToWorldScale, bVisibleViewDependent);
 }
 
 
@@ -54,17 +44,17 @@ FInputRayHit UGizmoElementCylinder::LineTrace(const FVector RayOrigin, const FVe
 		double RayParam = 0.0;
 		
 		const double PixelHitThresholdAdjust = CachedPixelToWorldScale * PixelHitDistanceThreshold;
-		const double CylinderHeight = Height * CachedLocalToWorldTransform.GetScale3D().X + PixelHitThresholdAdjust * 2.0;
-		const double CylinderRadius = Radius * CachedLocalToWorldTransform.GetScale3D().X + PixelHitThresholdAdjust;
-		const FVector CylinderDirection = CachedLocalToWorldTransform.TransformVectorNoScale(Direction);
-		const FVector CylinderLocalCenter = Base + Direction * Height * 0.5;
-		const FVector CylinderCenter = CachedLocalToWorldTransform.TransformPosition(CylinderLocalCenter);
+		const double WorldHeight = Height * CachedLocalToWorldTransform.GetScale3D().X + PixelHitThresholdAdjust * 2.0;
+		const double WorldRadius = Radius * CachedLocalToWorldTransform.GetScale3D().X + PixelHitThresholdAdjust;
+		const FVector WorldDirection = CachedLocalToWorldTransform.TransformVectorNoScale(Direction);
+		const FVector LocalCenter = Direction * Height * 0.5;
+		const FVector WorldCenter = CachedLocalToWorldTransform.TransformPosition(LocalCenter);
 
 		GizmoMath::RayCylinderIntersection(
-			CylinderCenter,
-			CylinderDirection,
-			CylinderRadius,
-			CylinderHeight,
+			WorldCenter,
+			WorldDirection,
+			WorldRadius,
+			WorldHeight,
 			RayOrigin, RayDirection,
 			bIntersects, RayParam);
 
diff --git a/Engine/Source/Runtime/InteractiveToolsFramework/Private/BaseGizmos/GizmoElementGroup.cpp b/Engine/Source/Runtime/InteractiveToolsFramework/Private/BaseGizmos/GizmoElementGroup.cpp
--- a/Engine/Source/Runtime/InteractiveToolsFramework/Private/BaseGizmos/GizmoElementGroup.cpp
+++ b/Engine/Source/Runtime/InteractiveToolsFramework/Private/BaseGizmos/GizmoElementGroup.cpp
@@ -9,30 +9,20 @@ void UGizmoElementGroup::Render(IToolsContextRenderAPI* RenderAPI, const FRender
 		return;
 	}
 
-	const FSceneView* View = RenderAPI->GetSceneView();
+	check(RenderAPI);
 
-	bool bVisibleViewDependent = GetViewDependentVisibility(View, RenderState.LocalToWorldTransform, FVector::ZeroVector);
-
-	FRenderTraversalState RenderStateCopy = RenderState;
+	FRenderTraversalState CurrentRenderState(RenderState);
+	bool bVisibleViewDependent = UpdateRenderState(RenderAPI, FVector::ZeroVector, CurrentRenderState);
 
 	if (bVisibleViewDependent)
 	{
 		// Compute constant scale, if applicable
-		float Scale = RenderStateCopy.LocalToWorldTransform.GetScale3D().X;
+		float Scale = CurrentRenderState.LocalToWorldTransform.GetScale3D().X;
 		if (bConstantScale)
 		{
-			Scale *= RenderStateCopy.PixelToWorldScale;
-		}
-		RenderStateCopy.LocalToWorldTransform.SetScale3D(FVector(Scale, Scale, Scale));
-
-		// Compute view alignment, if applicable
-		FQuat AlignRot;
-		if (GetViewAlignRot(View, RenderStateCopy.LocalToWorldTransform, FVector::ZeroVector, AlignRot))
-		{
-			RenderStateCopy.LocalToWorldTransform = FTransform(AlignRot) * RenderState.LocalToWorldTransform;
+			Scale *= CurrentRenderState.PixelToWorldScale;
 		}
-
-		UpdateRenderTraversalState(RenderStateCopy);
+		CurrentRenderState.LocalToWorldTransform.SetScale3D(FVector(Scale, Scale, Scale));
 
 		// Continue render even if not visible so all transforms will be cached 
 		// for subsequent line tracing.
@@ -40,12 +30,10 @@ void UGizmoElementGroup::Render(IToolsContextRenderAPI* RenderAPI, const FRender
 		{
 			if (Element)
 			{
-				Element->Render(RenderAPI, RenderStateCopy);
+				Element->Render(RenderAPI, CurrentRenderState);
 			}
 		}
 	}
-
-	CacheRenderState(RenderStateCopy.LocalToWorldTransform, RenderStateCopy.PixelToWorldScale, bVisibleViewDependent);
 }
 
 FInputRayHit UGizmoElementGroup::LineTrace(const FVector Start, const FVector Direction)
diff --git a/Engine/Source/Runtime/InteractiveToolsFramework/Private/BaseGizmos/GizmoElementLineBase.cpp b/Engine/Source/Runtime/InteractiveToolsFramework/Private/BaseGizmos/GizmoElementLineBase.cpp
--- a/Engine/Source/Runtime/InteractiveToolsFramework/Private/BaseGizmos/GizmoElementLineBase.cpp
+++ b/Engine/Source/Runtime/InteractiveToolsFramework/Private/BaseGizmos/GizmoElementLineBase.cpp
@@ -1,19 +1,25 @@
 // Copyright Epic Games, Inc. All Rights Reserved.
  
 #include "BaseGizmos/GizmoElementLineBase.h"
-#include "BaseGizmos/GizmoElementBase.h"
 
 DEFINE_LOG_CATEGORY_STATIC(LogGizmoElementLineBase, Log, All);
 
+bool UGizmoElementLineBase::UpdateRenderState(IToolsContextRenderAPI* RenderAPI, const FVector& InLocalOrigin, FRenderTraversalState& InOutRenderState)
+{
+	InOutRenderState.LineRenderState.Update(LineRenderAttributes);
+
+	return Super::UpdateRenderState(RenderAPI, InLocalOrigin, InOutRenderState);
+}
+
 float UGizmoElementLineBase::GetCurrentLineThickness() const
 {
 	if (ElementInteractionState == EGizmoElementInteractionState::Hovering)
 	{
-		return (LineThickness == 0.0f ? HoverLineThicknessMultiplier : LineThickness * HoverLineThicknessMultiplier);
+		return (LineThickness > 0.0f ? LineThickness * HoverLineThicknessMultiplier : HoverLineThicknessMultiplier);
 	}
 	else if (ElementInteractionState == EGizmoElementInteractionState::Interacting)
 	{
-		return (LineThickness == 0.0f ? InteractLineThicknessMultiplier : LineThickness * InteractLineThicknessMultiplier);
+		return (LineThickness > 0.0f ? LineThickness * InteractLineThicknessMultiplier : InteractLineThicknessMultiplier);
 	}
 
 	return LineThickness;
@@ -57,3 +63,78 @@ float UGizmoElementLineBase::GetInteractLineThicknessMultiplier() const
 	return InteractLineThicknessMultiplier;
 }
 
+void UGizmoElementLineBase::SetLineColor(FLinearColor InLineColor, bool InOverridesChildState)
+{
+	LineRenderAttributes.LineColor.SetColor(InLineColor, InOverridesChildState);
+}
+
+FLinearColor UGizmoElementLineBase::GetLineColor() const
+{
+	return LineRenderAttributes.LineColor.GetColor();
+}
+
+bool UGizmoElementLineBase::HasLineColor() const
+{
+	return LineRenderAttributes.LineColor.bHasValue;
+}
+
+bool UGizmoElementLineBase::GetLineColorOverridesChildState() const
+{
+	return LineRenderAttributes.LineColor.bOverridesChildState;
+}
+
+void UGizmoElementLineBase::ClearLineColor()
+{
+	LineRenderAttributes.LineColor.Reset();
+}
+
+void UGizmoElementLineBase::SetHoverLineColor(FLinearColor InHoverLineColor, bool InOverridesChildState)
+{
+	LineRenderAttributes.HoverLineColor.SetColor(InHoverLineColor, InOverridesChildState);
+}
+
+FLinearColor UGizmoElementLineBase::GetHoverLineColor() const
+{
+	return LineRenderAttributes.HoverLineColor.GetColor();
+}
+bool UGizmoElementLineBase::HasHoverLineColor() const
+{
+	return LineRenderAttributes.HoverLineColor.bHasValue;
+}
+
+bool UGizmoElementLineBase::GetHoverLineColorOverridesChildState() const
+{
+	return LineRenderAttributes.HoverLineColor.bOverridesChildState;
+}
+
+void UGizmoElementLineBase::ClearHoverLineColor()
+{
+	LineRenderAttributes.HoverLineColor.Reset();
+}
+
+void UGizmoElementLineBase::SetInteractLineColor(FLinearColor InInteractLineColor, bool InOverridesChildState)
+{
+	LineRenderAttributes.InteractLineColor.SetColor(InInteractLineColor, InOverridesChildState);
+}
+
+FLinearColor UGizmoElementLineBase::GetInteractLineColor() const
+{
+	return LineRenderAttributes.InteractLineColor.GetColor();
+}
+bool UGizmoElementLineBase::HasInteractLineColor() const
+{
+	return LineRenderAttributes.InteractLineColor.bHasValue;
+}
+
+bool UGizmoElementLineBase::GetInteractLineColorOverridesChildState() const
+{
+	return LineRenderAttributes.InteractLineColor.bOverridesChildState;
+}
+
+void UGizmoElementLineBase::ClearInteractLineColor()
+{
+	LineRenderAttributes.InteractLineColor.Reset();
+}
+
+
+
diff --git a/Engine/Source/Runtime/InteractiveToolsFramework/Private/BaseGizmos/GizmoElementRectangle.cpp b/Engine/Source/Runtime/InteractiveToolsFramework/Private/BaseGizmos/GizmoElementRectangle.cpp
--- a/Engine/Source/Runtime/InteractiveToolsFramework/Private/BaseGizmos/GizmoElementRectangle.cpp
+++ b/Engine/Source/Runtime/InteractiveToolsFramework/Private/BaseGizmos/GizmoElementRectangle.cpp
@@ -17,75 +17,58 @@ void UGizmoElementRectangle::Render(IToolsContextRenderAPI* RenderAPI, const FRe
 
 	check(RenderAPI);
 
-	const FSceneView* View = RenderAPI->GetSceneView();
-
-	FTransform LocalToWorldTransform = RenderState.LocalToWorldTransform;
-
-	bool bVisibleViewDependent = GetViewDependentVisibility(View, LocalToWorldTransform, Center);
+	FRenderTraversalState CurrentRenderState(RenderState);
+	bool bVisibleViewDependent = UpdateRenderState(RenderAPI, Center, CurrentRenderState);
 
 	if (bVisibleViewDependent)
 	{
-		FVector Axis0, Axis1;
-		FQuat AlignRot;
-		if (GetViewAlignRot(View, LocalToWorldTransform, Center, AlignRot))
-		{
-			Axis0 = AlignRot.RotateVector(UpDirection);
-			Axis1 = AlignRot.RotateVector(SideDirection);
-		}
-		else
-		{
-			Axis0 = UpDirection;
-			Axis1 = SideDirection;
-		}
-
-		Axis0 = LocalToWorldTransform.TransformVectorNoScale(Axis0);
-		Axis1 = LocalToWorldTransform.TransformVectorNoScale(Axis1);
-
-		FVector WorldCenter = LocalToWorldTransform.TransformPosition(Center);
-		float WorldWidth = Width * LocalToWorldTransform.GetScale3D().X;
-		float WorldHeight = Height * LocalToWorldTransform.GetScale3D().X;
+		const FVector WorldUpAxis = CurrentRenderState.LocalToWorldTransform.TransformVectorNoScale(UpDirection);
+		const FVector WorldSideAxis = CurrentRenderState.LocalToWorldTransform.TransformVectorNoScale(SideDirection);
+		const FVector WorldCenter = CurrentRenderState.LocalToWorldTransform.TransformPosition(FVector::ZeroVector);
+		const float WorldWidth = Width * CurrentRenderState.LocalToWorldTransform.GetScale3D().X;
+		const float WorldHeight = Height * CurrentRenderState.LocalToWorldTransform.GetScale3D().X;
 
 		FPrimitiveDrawInterface* PDI = RenderAPI->GetPrimitiveDrawInterface();
 
 		if (bDrawMesh)
 		{
-			if (const UMaterialInterface* UseMaterial = GetCurrentMaterial(RenderState))
+			if (const UMaterialInterface* UseMaterial = CurrentRenderState.GetCurrentMaterial())
 			{
-
-				DrawRectangleMesh(PDI, WorldCenter, Axis0, Axis1, VertexColor, WorldWidth, WorldHeight, UseMaterial->GetRenderProxy(), SDPG_Foreground);
+				FColor VertexColor = CurrentRenderState.GetVertexColor().ToFColor(false);
+				DrawRectangleMesh(PDI, WorldCenter, WorldUpAxis, WorldSideAxis, VertexColor, WorldWidth, WorldHeight, UseMaterial->GetRenderProxy(), SDPG_Foreground);
 			}
 		}
 		if (bDrawLine)
 		{
-			DrawRectangle(PDI, WorldCenter, Axis0, Axis1, LineColor, WorldWidth, WorldHeight, SDPG_Foreground, GetCurrentLineThickness());
+			FColor LineColor = CurrentRenderState.GetCurrentLineColor().ToFColor(false);
+			DrawRectangle(PDI, WorldCenter, WorldUpAxis, WorldSideAxis, LineColor, WorldWidth, WorldHeight, SDPG_Foreground, GetCurrentLineThickness());
 		}
 	}
-	CacheRenderState(LocalToWorldTransform, RenderState.PixelToWorldScale, bVisibleViewDependent);
 }
 
 FInputRayHit UGizmoElementRectangle::LineTrace(const FVector RayOrigin, const FVector RayDirection)
 {
 	if (IsHittableInView())
 	{
-		if (bHitMesh)
+		const FVector WorldUpAxis = CachedLocalToWorldTransform.TransformVectorNoScale(UpDirection);
+		const FVector WorldSideAxis = CachedLocalToWorldTransform.TransformVectorNoScale(SideDirection);
+		const FVector WorldNormal = FVector::CrossProduct(WorldUpAxis, WorldSideAxis);
+		const FVector WorldCenter = CachedLocalToWorldTransform.TransformPosition(FVector::ZeroVector);
+		const double Scale = CachedLocalToWorldTransform.GetScale3D().X;
+		const double PixelHitThresholdAdjust = CachedPixelToWorldScale * PixelHitDistanceThreshold;
+		const double WorldHeight = Scale * Height + PixelHitThresholdAdjust * 2.0;
+		const double WorldWidth = Scale * Width + PixelHitThresholdAdjust * 2.0;
+		const FVector Base = WorldCenter - WorldUpAxis * WorldHeight * 0.5 - WorldSideAxis * WorldWidth * 0.5;
+
+		// if ray is parallel to rectangle, no hit
+		if (FMath::IsNearlyZero(FVector::DotProduct(WorldNormal, RayDirection)))
 		{
-			const FVector UpAxis = CachedLocalToWorldTransform.TransformVectorNoScale(UpDirection);
-			const FVector SideAxis = CachedLocalToWorldTransform.TransformVectorNoScale(SideDirection);
-			const FVector Normal = FVector::CrossProduct(UpAxis, SideAxis);
-			const FVector WorldCenter = CachedLocalToWorldTransform.TransformPosition(Center);
-			const double Scale = CachedLocalToWorldTransform.GetScale3D().X;
-			const double PixelHitThresholdAdjust = CachedPixelToWorldScale * PixelHitDistanceThreshold;
-			const double WorldHeight = Scale * Height + 2.0 * PixelHitThresholdAdjust;
-			const double WorldWidth = Scale * Width + 2.0 * PixelHitThresholdAdjust;
-			const FVector Base = WorldCenter - UpAxis * WorldHeight * 0.5 - SideAxis * WorldWidth * 0.5;
-
-			// if ray is parallel to rectangle, no hit
-			if (FMath::IsNearlyZero(FVector::DotProduct(Normal, RayDirection)))
-			{
-				return FInputRayHit();
-			}
+			return FInputRayHit();
+		}
 
-			FPlane Plane(Base, Normal);
+		if (bHitMesh)
+		{
+			FPlane Plane(Base, WorldNormal);
 			double HitDepth = FMath::RayPlaneIntersectionParam(RayOrigin, RayDirection, Plane);
 			if (HitDepth < 0)
 			{
@@ -94,8 +77,8 @@ FInputRayHit UGizmoElementRectangle::LineTrace(const FVector RayOrigin, const FV
 
 			FVector HitPoint = RayOrigin + RayDirection * HitDepth;
 			FVector HitOffset = HitPoint - Base;
-			double HdU = FVector::DotProduct(HitOffset, UpAxis);
-			double HdS = FVector::DotProduct(HitOffset, SideAxis);
+			double HdU = FVector::DotProduct(HitOffset, WorldUpAxis);
+			double HdS = FVector::DotProduct(HitOffset, WorldSideAxis);
 
 			// clip to rectangle dimensions
 			if (HdU >= 0.0 && HdU <= WorldHeight && HdS >= 0.0 && HdS <= WorldWidth)
@@ -108,7 +91,30 @@ FInputRayHit UGizmoElementRectangle::LineTrace(const FVector RayOrigin, const FV
 		}
 		else if (bHitLine)
 		{
-			// @todo - add hit testing for line-drawn rectangle, requires storing PixelToWorld scale factor
+			FPlane Plane(Base, WorldNormal);
+			double HitDepth = FMath::RayPlaneIntersectionParam(RayOrigin, RayDirection, Plane);
+			if (HitDepth < 0)
+			{
+				return FInputRayHit();
+			}
+
+			FVector HitPoint = RayOrigin + RayDirection * HitDepth;
+			FVector HitOffset = HitPoint - Base;
+			double HdU = FVector::DotProduct(HitOffset, WorldUpAxis);
+			double HdS = FVector::DotProduct(HitOffset, WorldSideAxis);
+
+			const double HitBuffer = PixelHitThresholdAdjust + LineThickness;
+
+			// determine if the hit is within pixel tolerance of the edges of rectangle
+			if (HdU >= 0.0 && HdU <= WorldHeight && HdS >= 0.0 && HdS <= WorldWidth &&
+				(HdS <= HitBuffer || HdS >= WorldWidth - HitBuffer ||
+				 HdU <= HitBuffer || HdU >= WorldHeight - HitBuffer))
+			{
+				FInputRayHit RayHit(HitDepth);
+				RayHit.SetHitObject(this);
+				RayHit.HitIdentifier = PartIdentifier;
+				return RayHit;
+			}
 		}
 	}
 	return FInputRayHit();
@@ -170,16 +176,6 @@ FVector UGizmoElementRectangle::GetSideDirection() const
 	return SideDirection;
 }
 
-void UGizmoElementRectangle::SetLineColor(const FColor& InLineColor)
-{
-	LineColor = InLineColor;
-}
-
-FColor UGizmoElementRectangle::GetLineColor() const
-{
-	return LineColor;
-}
-
 void UGizmoElementRectangle::SetDrawMesh(bool InDrawMesh)
 {
 	bDrawMesh = InDrawMesh;
diff --git a/Engine/Source/Runtime/InteractiveToolsFramework/Private/BaseGizmos/GizmoElementRenderState.cpp b/Engine/Source/Runtime/InteractiveToolsFramework/Private/BaseGizmos/GizmoElementRenderState.cpp
new file mode 100644
--- /dev/null
+++ b/Engine/Source/Runtime/InteractiveToolsFramework/Private/BaseGizmos/GizmoElementRenderState.cpp
@@ -0,0 +1,115 @@
+// Copyright Epic Games, Inc. All Rights Reserved.
+ 
+#include "BaseGizmos/GizmoElementRenderState.h"
+#include "Materials/MaterialInterface.h"
+
+FLinearColor FGizmoElementColorAttribute::GetColor() const
+{
+	if (bHasValue)
+	{
+		return Value;
+	}
+	return DefaultColor;
+}
+
+void FGizmoElementColorAttribute::SetColor(FLinearColor InColor, bool InOverridesChildState)
+{
+	Value = InColor;
+	bHasValue = true;
+	bOverridesChildState = InOverridesChildState;
+}
+
+void FGizmoElementColorAttribute::Reset()
+{
+	Value = DefaultColor;
+	bHasValue = false;
+	bOverridesChildState = false;
+}
+
+void  FGizmoElementColorAttribute::UpdateState(const FGizmoElementColorAttribute& InChildColorAttribute)
+{
+	if (InChildColorAttribute.bHasValue && !(bOverridesChildState && bHasValue))
+	{
+		Value = InChildColorAttribute.Value;
+		bHasValue = true;
+		bOverridesChildState = InChildColorAttribute.bOverridesChildState;
+	}
+}
+
+const UMaterialInterface* FGizmoElementMaterialAttribute::GetMaterial() const
+{
+	if (Value.IsValid())
+	{
+		return Value.Get();
+	}
+	return nullptr;
+}
+
+void FGizmoElementMaterialAttribute::SetMaterial(TWeakObjectPtr<UMaterialInterface> InMaterial, bool InOverridesChildState)
+{
+	Value = InMaterial;
+	bOverridesChildState = InOverridesChildState;
+}
+
+void FGizmoElementMaterialAttribute::Reset()
+{
+	Value = nullptr;
+	bOverridesChildState = false;
+}
+
+void FGizmoElementMaterialAttribute::UpdateState(const FGizmoElementMaterialAttribute& InChildMaterialAttribute)
+{
+	if (InChildMaterialAttribute.Value != nullptr && !(bOverridesChildState && Value != nullptr))
+	{
+		Value = InChildMaterialAttribute.Value;
+		bOverridesChildState = InChildMaterialAttribute.bOverridesChildState;
+	}
+}
+
+const UMaterialInterface* FGizmoElementMeshRenderStateAttributes::GetMaterial(EGizmoElementInteractionState InteractionState)
+{
+	if (InteractionState == EGizmoElementInteractionState::Hovering)
+	{
+		return HoverMaterial.GetMaterial();
+	}
+	else if (InteractionState == EGizmoElementInteractionState::Interacting)
+	{
+		return InteractMaterial.GetMaterial();
+	}
+
+	return Material.GetMaterial();
+}
+
+FLinearColor FGizmoElementMeshRenderStateAttributes::GetVertexColor()
+{
+	return VertexColor.GetColor();
+}
+
+void FGizmoElementMeshRenderStateAttributes::Update(FGizmoElementMeshRenderStateAttributes& InChildAttributes)
+{
+	Material.UpdateState(InChildAttributes.Material);
+	HoverMaterial.UpdateState(InChildAttributes.HoverMaterial);
+	InteractMaterial.UpdateState(InChildAttributes.InteractMaterial);
+	VertexColor.UpdateState(InChildAttributes.VertexColor);
+}
+
+FLinearColor FGizmoElementLineRenderStateAttributes::GetLineColor(EGizmoElementInteractionState InteractionState)
+{
+	if (InteractionState == EGizmoElementInteractionState::Hovering)
+	{
+		return HoverLineColor.GetColor();
+	}
+	else if (InteractionState == EGizmoElementInteractionState::Interacting)
+	{
+		return InteractLineColor.GetColor();
+	}
+
+	return LineColor.GetColor();
+}
+
+void FGizmoElementLineRenderStateAttributes::Update(FGizmoElementLineRenderStateAttributes& InChildAttributes)
+{
+	LineColor.UpdateState(InChildAttributes.LineColor);
+	HoverLineColor.UpdateState(InChildAttributes.HoverLineColor);
+	InteractLineColor.UpdateState(InChildAttributes.InteractLineColor);
+}
diff --git a/Engine/Source/Runtime/InteractiveToolsFramework/Private/BaseGizmos/GizmoElementTorus.cpp b/Engine/Source/Runtime/InteractiveToolsFramework/Private/BaseGizmos/GizmoElementTorus.cpp
--- a/Engine/Source/Runtime/InteractiveToolsFramework/Private/BaseGizmos/GizmoElementTorus.cpp
+++ b/Engine/Source/Runtime/InteractiveToolsFramework/Private/BaseGizmos/GizmoElementTorus.cpp
@@ -15,41 +15,23 @@ void UGizmoElementTorus::Render(IToolsContextRenderAPI* RenderAPI, const FRender
 	}
 
 	check(RenderAPI);
-	const FSceneView* View = RenderAPI->GetSceneView();
 
-	FTransform LocalToWorldTransform = RenderState.LocalToWorldTransform;
 
-	bool bVisibleViewDependent = GetViewDependentVisibility(View, LocalToWorldTransform, Center);
+	FRenderTraversalState CurrentRenderState(RenderState);
+	bool bVisibleViewDependent = UpdateRenderState(RenderAPI, Center, CurrentRenderState);
 
 	if (bVisibleViewDependent)
 	{
-		const UMaterialInterface* UseMaterial = GetCurrentMaterial(RenderState);
-			
-		if (UseMaterial)
+		if (const UMaterialInterface* UseMaterial = CurrentRenderState.GetCurrentMaterial())
 		{
 			FPrimitiveDrawInterface* PDI = RenderAPI->GetPrimitiveDrawInterface();
 
 			FVector TorusSideAxis = Normal ^ BeginAxis;
 			TorusSideAxis.Normalize();
 
-			FTransform TorusToLocal = FTransform::Identity;
-			TorusToLocal.SetTranslation(Center);
-
-			FQuat AlignRot;
-			if (GetViewAlignRot(View, LocalToWorldTransform, Center, AlignRot))
-			{
-				TorusToLocal.SetRotation(AlignRot);
-			}
-
-			FVector WorldAlignAxis = LocalToWorldTransform.TransformVector(ViewAlignAxis);
-
-			LocalToWorldTransform = TorusToLocal * LocalToWorldTransform;
-
-			DrawTorus(PDI, LocalToWorldTransform.ToMatrixWithScale(), BeginAxis, TorusSideAxis, OuterRadius, InnerRadius, OuterSegments, InnerSlices,  UseMaterial->GetRenderProxy(), SDPG_Foreground, bPartial, Angle, bEndCaps);
+			DrawTorus(PDI, CurrentRenderState.LocalToWorldTransform.ToMatrixWithScale(), BeginAxis, TorusSideAxis, OuterRadius, InnerRadius, OuterSegments, InnerSlices,  UseMaterial->GetRenderProxy(), SDPG_Foreground, bPartial, Angle, bEndCaps);
 		}
 	}
-
-	CacheRenderState(LocalToWorldTransform, bVisibleViewDependent);
 }
 
 
diff --git a/Engine/Source/Runtime/InteractiveToolsFramework/Public/BaseGizmos/GizmoElementBase.h b/Engine/Source/Runtime/InteractiveToolsFramework/Public/BaseGizmos/GizmoElementBase.h
--- a/Engine/Source/Runtime/InteractiveToolsFramework/Public/BaseGizmos/GizmoElementBase.h
+++ b/Engine/Source/Runtime/InteractiveToolsFramework/Public/BaseGizmos/GizmoElementBase.h
@@ -3,6 +3,8 @@
 #pragma once
 
 #include "CoreMinimal.h"
+#include "BaseGizmos/GizmoElementRenderState.h"
+#include "BaseGizmos/GizmoElementShared.h"
 #include "BaseGizmos/GizmoInterfaces.h"
 #include "BaseGizmos/GizmoRenderingUtil.h"
 #include "InputState.h"
@@ -13,42 +15,6 @@
 
 class UMaterialInterface;
 
-// bitmask indicating whether object is visible, hittable or both
-UENUM()
-enum class EGizmoElementState : uint8
-{
-	None		= 0x00,
-	Visible		= 1<<1,
-	Hittable	= 1<<2,
-	VisibleAndHittable = Visible | Hittable
-};
-
-UENUM()
-enum class EGizmoElementInteractionState
-{
-	None,
-	Hovering,
-	Interacting
-};
-
-UENUM()
-enum class EGizmoElementViewDependentType
-{
-	None,
-	Axis, // Cull object when angle between axis and view direction is within a given tolerance
-	Plane // Cull object when angle between plane normal and view direction is perpendicular within a given tolerance
-};
-
-UENUM()
-enum class EGizmoElementViewAlignType
-{
-	None,
-	PointOnly, // Align object forward axis to view direction only, useful for symmetrical objects such as a circle
-	PointEye, // Align object forward axis to -camera view direction (camera pos - object center), align object up axis to scene view up
-	PointScreen, // Align object forward axis to scene view forward direction (view up ^ view right), align object up axis to scene view up
-	Axial // Rotate object around up axis, minimizing angle between forward axis and view direction
-};
-
 /**
  * Base class for 2d and 3d primitive objects intended to be used as part of 3D Gizmos.
  * Contains common properties and utility functions.
@@ -69,7 +35,15 @@ public:
 
 	static constexpr uint32 DefaultPartIdentifier = 0;						// Default part ID, used for elements that are not associated with any gizmo part
 
-	// Helper struct used to store traversal state during render
+	// 
+	// Render traversal state structure used to maintain the current render state while rendering.
+	// As the gizmo element hierarchy is traversed, current state is maintained and updated. 
+	// Element state attribute inheritance works as follows:
+	//
+	// - Child element state that is not set inherits from parent state.
+	// - Child element state that is set replaces the parent state, except in the case of overrides.
+	// - Overrides: parent element state with override set to true replaces all child state regardless of whether the child state has been set.
+	//
 	struct FRenderTraversalState
 	{
 		// LocalToWorld transform 
@@ -79,23 +53,39 @@ public:
 		// Pixel to world scale
 		double PixelToWorldScale = 1.0;
 
-		// Material, if set, overrides the element's material
-		TWeakObjectPtr<UMaterialInterface> Material = nullptr;
-
-		// Hover material, if set, overrides the element's hover material
-		TWeakObjectPtr<UMaterialInterface> HoverMaterial = nullptr;
+		// Interact state, if not equal to none, overrides the element's interact state 
+		EGizmoElementInteractionState InteractionState = EGizmoElementInteractionState::None;
 
-		// Interact material, if set, overrides the element's interact material
-		TWeakObjectPtr<UMaterialInterface> InteractMaterial = nullptr;
+		// Current state used for rendering meshes
+		FGizmoElementMeshRenderStateAttributes MeshRenderState;
 
-		// Interact state, if not equal to none, overrides the element's interact state
-		EGizmoElementInteractionState InteractionState = EGizmoElementInteractionState::None;
+		// Current state used for rendering lines
+		FGizmoElementLineRenderStateAttributes LineRenderState;
 
+		// Initialize state 
 		void Initialize(const FSceneView* InSceneView, FTransform InTransform)
 		{
 			LocalToWorldTransform = InTransform;
 			PixelToWorldScale = GizmoRenderingUtil::CalculateLocalPixelToWorldScale(InSceneView, InTransform.GetLocation());
 		}
+
+		// Returns the mesh material based on the current interaction state.
+		const UMaterialInterface* GetCurrentMaterial()
+		{
+			return MeshRenderState.GetMaterial(InteractionState);
+		}
+
+		// Returns the mesh vertex color.
+		FLinearColor GetVertexColor()
+		{
+			return MeshRenderState.GetVertexColor();
+		}
+
+		// Returns the line color based on the current interaction state.
+		FLinearColor GetCurrentLineColor()
+		{
+			return LineRenderState.GetLineColor(InteractionState);
+		}
 	};
 
 public:
@@ -107,10 +97,7 @@ public:
 	virtual FInputRayHit LineTrace(const FVector Start, const FVector Direction) PURE_VIRTUAL(UGizmoElementBase::LineTrace, return FInputRayHit(););
 
 	// Calcute box sphere bounds for use when hit testing.
-	virtual FBoxSphereBounds CalcBounds(const FTransform & LocalToWorld) const PURE_VIRTUAL(UGizmoElementBase::CalcBounds, return FBoxSphereBounds(););
-
-	// Returns the current material for given object state. The materials in the render state will override those in the element.
-	virtual const UMaterialInterface* GetCurrentMaterial(const FRenderTraversalState& RenderState) const;
+	virtual FBoxSphereBounds CalcBounds(const FTransform& LocalToWorld) const PURE_VIRTUAL(UGizmoElementBase::CalcBounds, return FBoxSphereBounds(););
 
 	// Reset the cached render state.
 	virtual void ResetCachedRenderState();
@@ -173,13 +160,14 @@ public:
 	virtual EGizmoElementViewAlignType GetViewAlignType() const;
 
 	// View align axis. 
-	// PointEye and PointOnly rotate this axis to align with view up.
+	// PointEye, PointScreen and PointOnly rotate this axis to align with view up.
 	// Axial rotates about this axis.
 	virtual void SetViewAlignAxis(FVector InAxis);
 	virtual FVector GetViewAlignAxis() const;
 
 	// View align normal.
-	// PointEye rotates the normal to align with view direction.
+	// PointEye rotates the normal to align with camera view direction.
+	// PointScreen rotates the normal to align with screen forward direction.
 	// Axial rotates the normal around the axis to align as closely as possible with the view direction.
 	virtual void SetViewAlignNormal(FVector InAxis);
 	virtual FVector GetViewAlignNormal() const;
@@ -189,25 +177,77 @@ public:
 	virtual void SetViewAlignAxialAngleTol(float InMaxAngleTol);
 	virtual float GetViewAlignAxialAngleTol() const;
 
-	// Default material.
-	virtual void SetMaterial(UMaterialInterface* Material);
-	UMaterialInterface* GetMaterial() const;
+	// Pixel hit distance threshold, element will be scaled enough to add this threshold when line-tracing. */
+	virtual void SetPixelHitDistanceThreshold(float InPixelHitDistanceThreshold);
+	virtual float GetPixelHitDistanceThreshold() const;
 
-	// Hover material used when object state is hovering.
-	virtual void SetHoverMaterial(UMaterialInterface* Material);
-	UMaterialInterface* GetHoverMaterial() const;
+	//
+	// Methods for managing render state attributes: Material, HoverMaterial, InteractMaterial, VertexColor 
+	// 
+	// State inheritance works as follows: 
+	// - Gizmo element state that is not set inherits from the corresponding state in the current render traversal.
+	// - Gizmo element state that is set replaces the corresponding state in the current render traversal, except in the case of overrides.
+	// - Gizmo element state that is set to override, will override any corresponding state in children.
+	//
 
-	// Interact material used when object state is interacting.
-	virtual void SetInteractMaterial(UMaterialInterface* Material);
-	UMaterialInterface* GetInteractMaterial() const;
+	// Set mesh render state material attribute. 
+	//  @param InMaterial - material to be set
+	//  @param InOverridesChildState - when true, this material will override the material of all child elements.
+	virtual void SetMaterial(TWeakObjectPtr<UMaterialInterface> InMaterial, bool InOverridesChildState = false);
 
-	// Vertex color where applicable.
-	virtual void SetVertexColor(const FColor& InVertexColor);
-	FColor GetVertexColor() const;
+	// Get mesh render state material attribute's value. 
+	virtual const UMaterialInterface* GetMaterial() const;
 
-	// Pixel hit distance threshold, element will be scaled enough to add this threshold when line-tracing.
-	virtual void SetPixelHitDistanceThreshold(float InPixelHitDistanceThreshold);
-	virtual float GetPixelHitDistanceThreshold() const;
+	// Get mesh render state material attribute's override setting. 
+	virtual bool GetMaterialOverridesChildState() const;
+
+	// Clear mesh render state material attribute. 
+	virtual void ClearMaterial();
+
+	// Set mesh render state hover material attribute. 
+	//  @param InHoverMaterial - hover material to be set
+	//  @param InOverridesChildState - when true, this hover material will override the material of all child elements.
+	virtual void SetHoverMaterial(TWeakObjectPtr<UMaterialInterface> InHoverMaterial, bool InOverridesChildState = false);
+
+	// Get mesh render state hover material attribute's value.
+	virtual const UMaterialInterface* GetHoverMaterial() const;
+
+	// Get mesh render state hover material attribute's override setting.
+	virtual bool GetHoverMaterialOverridesChildState() const;
+
+	// Clear mesh render state hover material attribute. 
+	virtual void ClearHoverMaterial();
+
+	// Set mesh render state interact material attribute. 
+	//  @param InHoverMaterial - interact material to be set
+	//  @param InOverridesChildState - when true, this interact material will override the material of all child elements.
+	virtual void SetInteractMaterial(TWeakObjectPtr<UMaterialInterface> InInteractMaterial, bool InOverridesChildState = false);
+
+	// Get mesh render state interact material attribute's value. 
+	virtual const UMaterialInterface* GetInteractMaterial() const;
+
+	// Get mesh render state interact material attribute's override setting. 
+	virtual bool GetInteractMaterialOverridesChildState() const;
+
+	// Clear mesh render interact state material attribute. 
+	virtual void ClearInteractMaterial();
+
+	// Set mesh render state vertex color attribute. 
+	//  @param InVertexColor - vertex color to be set
+	//  @param InOverridesChildState - when true, this vertex color will override the material of all child elements.
+	virtual void SetVertexColor(FLinearColor InVertexColor, bool InOverridesChildState = false);
+
+	// Get mesh render state vertex color attribute's value. 
+	virtual FLinearColor GetVertexColor() const;
+
+	// Returns true, if mesh render state vertex color attribute has been set. 
+	virtual bool HasVertexColor() const;
+
+	// Get mesh render state vertex color attribute's override setting. 
+	virtual bool GetVertexColorOverridesChildState() const;
+
+	// Clear mesh render state vertex color attribute.
+	virtual void ClearVertexColor();
 
 protected:
 
@@ -219,6 +259,10 @@ protected:
 	UPROPERTY()
 	uint32 PartIdentifier = DefaultPartIdentifier;
 
+	// Mesh render state attributes for this element
+	UPROPERTY()
+	FGizmoElementMeshRenderStateAttributes MeshRenderAttributes;
+
 	// Element state - indicates whether object is visible or hittable
 	UPROPERTY()
 	EGizmoElementState ElementState = EGizmoElementState::VisibleAndHittable;
@@ -279,22 +323,6 @@ protected:
 	UPROPERTY()
 	float ViewAlignAxialMaxCosAngleTol = DefaultViewAlignMaxCosAngleTol;
 
-	// Default material.
-	UPROPERTY()
-	TObjectPtr<UMaterialInterface> Material;
-
-	// Hovering material.
-	UPROPERTY()
-	TObjectPtr<UMaterialInterface> HoverMaterial;
-
-	// Interacting material.
-	UPROPERTY()
-	TObjectPtr<UMaterialInterface> InteractMaterial;
-
-	// Vertex color where applicable.
-	UPROPERTY()
-	FColor VertexColor = FColor::White;
-
 	// Pixel hit distance threshold, element will be scaled enough to add this threshold when line-tracing.
 	UPROPERTY()
 	float PixelHitDistanceThreshold = 7.0;
@@ -335,11 +363,12 @@ protected:
 	// Helper method to calculate rotation between coord spaces.
 	FQuat GetAlignRotBetweenCoordSpaces(FVector SourceForward, FVector SourceSide, FVector SourceUp, FVector TargetForward, FVector TargetSide, FVector TargetUp) const;
 
-	// Cache render state during render traversal, to be used subsequently when line tracing.
-	virtual void CacheRenderState(const FTransform& InLocalToWorldState, double InPixelToWorldScale, bool InVisibleViewDependent = true);
+	// Update render state during render traversal, determines the current render state for this element 
+	// @return view dependent visibility, true if this element is visible in the current view. 
+	virtual bool UpdateRenderState(IToolsContextRenderAPI* RenderAPI, const FVector& InLocalOrigin, FRenderTraversalState& InOutRenderState);
 
-	// Update render traversal state based on properties in the current object, should always be called at the beginning of Render()
-	virtual void UpdateRenderTraversalState(FRenderTraversalState& InRenderTraversalState);
+	// Cache render state during render traversal to be used subsequently when line tracing.
+	virtual void CacheRenderState(const FTransform& InLocalToWorldState, double InPixelToWorldScale, bool InVisibleViewDependent = true);
 
 };
 
diff --git a/Engine/Source/Runtime/InteractiveToolsFramework/Public/BaseGizmos/GizmoElementCircle.h b/Engine/Source/Runtime/InteractiveToolsFramework/Public/BaseGizmos/GizmoElementCircle.h
--- a/Engine/Source/Runtime/InteractiveToolsFramework/Public/BaseGizmos/GizmoElementCircle.h
+++ b/Engine/Source/Runtime/InteractiveToolsFramework/Public/BaseGizmos/GizmoElementCircle.h
@@ -40,10 +40,6 @@ public:
 	virtual void SetNumSides(int InNumSides);
 	virtual int GetNumSides() const;
 
-	// Line color when bDrawLine is true.
-	virtual void SetLineColor(const FLinearColor& InColor);
-	virtual FLinearColor GetLineColor() const;
-
 	// Draw mesh
 	virtual void SetDrawMesh(bool InDrawMesh);
 	virtual bool GetDrawMesh() const;
@@ -78,10 +74,6 @@ protected:
 	UPROPERTY()
 	int NumSides = 64;
 
-	// Line color when bFill is false.
-	UPROPERTY()
-	FLinearColor LineColor = FLinearColor::White;
-
 	// Whether to render solid circle.
 	UPROPERTY()
 	bool bDrawMesh = true;
diff --git a/Engine/Source/Runtime/InteractiveToolsFramework/Public/BaseGizmos/GizmoElementGroup.h b/Engine/Source/Runtime/InteractiveToolsFramework/Public/BaseGizmos/GizmoElementGroup.h
--- a/Engine/Source/Runtime/InteractiveToolsFramework/Public/BaseGizmos/GizmoElementGroup.h
+++ b/Engine/Source/Runtime/InteractiveToolsFramework/Public/BaseGizmos/GizmoElementGroup.h
@@ -4,7 +4,7 @@
 
 #include "CoreMinimal.h"
 #include "UObject/ObjectMacros.h"
-#include "BaseGizmos/GizmoElementBase.h"
+#include "BaseGizmos/GizmoElementLineBase.h"
 #include "InputState.h"
 #include "GizmoElementGroup.generated.h"
 
@@ -13,7 +13,7 @@
  * Contains multiple gizmo objects.
  */
 UCLASS(Transient)
-class INTERACTIVETOOLSFRAMEWORK_API UGizmoElementGroup : public UGizmoElementBase
+class INTERACTIVETOOLSFRAMEWORK_API UGizmoElementGroup : public UGizmoElementLineBase
 {
 	GENERATED_BODY()
 
diff --git a/Engine/Source/Runtime/InteractiveToolsFramework/Public/BaseGizmos/GizmoElementLineBase.h b/Engine/Source/Runtime/InteractiveToolsFramework/Public/BaseGizmos/GizmoElementLineBase.h
--- a/Engine/Source/Runtime/InteractiveToolsFramework/Public/BaseGizmos/GizmoElementLineBase.h
+++ b/Engine/Source/Runtime/InteractiveToolsFramework/Public/BaseGizmos/GizmoElementLineBase.h
@@ -4,6 +4,9 @@
 
 #include "CoreMinimal.h"
 #include "BaseGizmos/GizmoElementBase.h"
+#include "BaseGizmos/GizmoElementRenderState.h"
+#include "BaseGizmos/GizmoElementShared.h"
+#include "ToolContextInterfaces.h"
 #include "GizmoElementLineBase.generated.h"
 
 /**
@@ -31,8 +34,72 @@ public:
 	virtual void SetInteractLineThicknessMultiplier(float InInteractLineThicknessMultiplier);
 	virtual float GetInteractLineThicknessMultiplier() const;
 
+	//
+	// Methods for managing line state attributes: LineColor, HoverLineColor, InteractLineColor
+	//
+	// State inheritance works as follows:
+	// - Gizmo element state that is not set inherits from the corresponding state in the current render traversal.
+	// - Gizmo element state that is set replaces the corresponding state in the current render traversal, except in the case of overrides.
+	// - Gizmo element state that is set to override, will override any corresponding state in children.
+	//
+
+	// Set line render state line color attribute. 
+	//  @param InMaterial - line color to be set
+	//  @param InOverridesChildState - when true, this line color will override the line color of all child elements.
+	virtual void SetLineColor(FLinearColor InColor, bool InOverridesChildState = false);
+
+	// Get line render state line color attribute's value.
+	virtual FLinearColor GetLineColor() const;
+
+	// Returns true if line render state line color attribute has been set.
+	virtual bool HasLineColor() const;
+
+	// Get line render state line color attribute's override setting. 
+	virtual bool GetLineColorOverridesChildState() const;
+
+	// Clear line render state line color attribute.
+	virtual void ClearLineColor();
+
+	// Set line render state hover line color attribute. 
+	//  @param InMaterial - hover line color to be set
+	//  @param InOverridesChildState - when true, this hover line color will override the line color of all child elements.
+	virtual void SetHoverLineColor(FLinearColor InColor, bool InOverridesChildState = false);
+
+	// Get line render state hover line color attribute's value.
+	virtual FLinearColor GetHoverLineColor() const;
+
+	// Returns true if line render state hover line color attribute has been set. 
+	virtual bool HasHoverLineColor() const;
+
+	// Get line render state hover line color attribute's override setting.
+	virtual bool GetHoverLineColorOverridesChildState() const;
+
+	// Clear line render state hover line color attribute. 
+	virtual void ClearHoverLineColor();
+
+	// Set line render state interact line color attribute. 
+	//  @param InMaterial - interact line color to be set
+	//  @param InOverridesChildState - when true, this interact line color will override the line color of all child elements.
+	virtual void SetInteractLineColor(FLinearColor InColor, bool InOverridesChildState = false);
+
+	// Get line render state interact line color attribute's value. 
+	virtual FLinearColor GetInteractLineColor() const;
+
+	// Returns true if line render state interact line color attribute has been set. 
+	virtual bool HasInteractLineColor() const;
+
+	// Get line render state interact line color attribute's override setting. 
+	virtual bool GetInteractLineColorOverridesChildState() const;
+
+	// Clear line render state interact line color attribute.
+	virtual void ClearInteractLineColor();
+
 protected:
 
+	// Line render state attributes for this element
+	UPROPERTY()
+	FGizmoElementLineRenderStateAttributes LineRenderAttributes;
+
 	// Line thickness when rendering lines, must be >= 0.0, value of 0.0 will render thinnest line 
 	UPROPERTY()
 	float LineThickness = 0.0;
@@ -44,4 +111,8 @@ protected:
 	// Multiplier applied to line thickness when interacting
 	UPROPERTY(EditAnywhere, Category = Options)
 	float InteractLineThicknessMultiplier = 2.0f;
+
+	// Update render state during render traversal, determines the current render state for this element
+	// @return view dependent visibility, true if this element is visible in the current view.
+	virtual bool UpdateRenderState(IToolsContextRenderAPI* RenderAPI, const FVector& InLocalOrigin, FRenderTraversalState& InOutRenderState) override;
 };
\ No newline at end of file
diff --git a/Engine/Source/Runtime/InteractiveToolsFramework/Public/BaseGizmos/GizmoElementRectangle.h b/Engine/Source/Runtime/InteractiveToolsFramework/Public/BaseGizmos/GizmoElementRectangle.h
--- a/Engine/Source/Runtime/InteractiveToolsFramework/Public/BaseGizmos/GizmoElementRectangle.h
+++ b/Engine/Source/Runtime/InteractiveToolsFramework/Public/BaseGizmos/GizmoElementRectangle.h
@@ -48,10 +48,6 @@ public:
 	virtual void SetSideDirection(const FVector& InSideDirection);
 	virtual FVector GetSideDirection() const;
 
-	// Line color
-	virtual void SetLineColor(const FColor& InLineColor);
-	virtual FColor GetLineColor() const;
-
 	// Draw mesh
 	virtual void SetDrawMesh(bool InDrawMesh);
 	virtual bool GetDrawMesh() const;
@@ -90,10 +86,6 @@ protected:
 	UPROPERTY()
 	FVector SideDirection = FVector(0.0f, 1.0f, 0.0f);
 
-	// Line color when bShowOutline is false.
-	UPROPERTY()
-	FColor LineColor = FColor::White;
-
 	UPROPERTY()
 	bool bDrawMesh = true;
 
diff --git a/Engine/Source/Runtime/InteractiveToolsFramework/Public/BaseGizmos/GizmoElementRenderState.h b/Engine/Source/Runtime/InteractiveToolsFramework/Public/BaseGizmos/GizmoElementRenderState.h
new file mode 100644
--- /dev/null
+++ b/Engine/Source/Runtime/InteractiveToolsFramework/Public/BaseGizmos/GizmoElementRenderState.h
@@ -0,0 +1,145 @@
+// Copyright Epic Games, Inc. All Rights Reserved.
+
+#pragma once
+
+#include "CoreMinimal.h"
+#include "UObject/Object.h"
+#include "UObject/ObjectMacros.h"
+#include "BaseGizmos/GizmoElementShared.h"
+#include "Materials/MaterialInterface.h"
+#include "GizmoElementRenderState.generated.h"
+
+
+//
+// Structs used to store gizmo element render state.
+//
+
+//
+// Color state attribute struct 
+//
+// Stores color state value, whether state is set and whether it overrides child state. 
+//
+USTRUCT()
+struct INTERACTIVETOOLSFRAMEWORK_API FGizmoElementColorAttribute
+{
+	GENERATED_BODY()
+
+	static constexpr FLinearColor DefaultColor = FLinearColor(1.0f, 1.0f, 1.0f);
+
+	// Color value
+	UPROPERTY()
+	FLinearColor Value = DefaultColor;
+
+	// True if value has been set.
+	UPROPERTY()
+	bool bHasValue = false;
+
+	// True if this attribute overrides child state attributes
+	UPROPERTY()
+	bool bOverridesChildState = false;
+
+	// Return the color value. This returns the default color value if no value has been set.
+	FLinearColor GetColor() const;
+
+	// Set color value and override flag
+	void SetColor(FLinearColor InColor, bool InOverridesChildState = false);
+
+	// Reset attribute to default values.
+	void Reset();
+
+	// Update this state based on a child state attribute.
+	void UpdateState(const FGizmoElementColorAttribute& InChildColorAttribute);
+};
+
+//
+// Material state attribute struct
+//
+// Stores material state value, whether state is set and whether it overrides child state.
+//
+USTRUCT()
+struct INTERACTIVETOOLSFRAMEWORK_API FGizmoElementMaterialAttribute 
+{
+	GENERATED_BODY()
+
+	// Material value
+	UPROPERTY()
+	TWeakObjectPtr<UMaterialInterface> Value = nullptr;
+
+	// True if this attribute overrides child state attributes 
+	UPROPERTY()
+	bool bOverridesChildState = false;
+
+	// Return the material value.
+	const UMaterialInterface* GetMaterial() const;
+
+	// Set material value and override flag
+	void SetMaterial(TWeakObjectPtr<UMaterialInterface> InColor, bool InOverridesChildState = false);
+
+	// Reset attribute to default values.
+	void Reset();
+
+	// Update this state based on a child state attribute.
+	void UpdateState(const FGizmoElementMaterialAttribute& InChildMaterialAttribute);
+};
+
+//
+// Mesh render state structure.
+//
+USTRUCT()
+struct INTERACTIVETOOLSFRAMEWORK_API FGizmoElementMeshRenderStateAttributes
+{
+	GENERATED_BODY()
+
+	// Default material
+	UPROPERTY()
+	FGizmoElementMaterialAttribute Material;
+
+	// Hover material
+	UPROPERTY()
+	FGizmoElementMaterialAttribute HoverMaterial;
+
+	// Interact material
+	UPROPERTY()
+	FGizmoElementMaterialAttribute InteractMaterial;
+
+	// Mesh vertex color
+	UPROPERTY()
+	FGizmoElementColorAttribute VertexColor;
+
+	// Returns the material corresponding to the input interaction state
+	const UMaterialInterface* GetMaterial(EGizmoElementInteractionState InteractionState);
+
+	// Returns mesh vertex color
+	FLinearColor GetVertexColor();
+
+	// Update this mesh render state based on a child mesh render state attribute.
+	void Update(FGizmoElementMeshRenderStateAttributes& InChildAttributes);
+};
+
+//
+// Line render state structure.
+//
+USTRUCT()
+struct INTERACTIVETOOLSFRAMEWORK_API FGizmoElementLineRenderStateAttributes
+{
+	GENERATED_BODY()
+
+	// Line color
+	UPROPERTY()
+	FGizmoElementColorAttribute LineColor;
+
+	// Hover line color
+	UPROPERTY()
+	FGizmoElementColorAttribute HoverLineColor;
+
+	// Interact line color
+	UPROPERTY()
+	FGizmoElementColorAttribute InteractLineColor;
+
+	// Returns the line color corresponding to the input interaction state
+	FLinearColor GetLineColor(EGizmoElementInteractionState InteractionState);
+
+	// Update this line render state based on a child mesh render state attribute.
+	void Update(FGizmoElementLineRenderStateAttributes& InChildAttributes);
+};
+
diff --git a/Engine/Source/Runtime/InteractiveToolsFramework/Public/BaseGizmos/GizmoElementShared.h b/Engine/Source/Runtime/InteractiveToolsFramework/Public/BaseGizmos/GizmoElementShared.h
new file mode 100644
--- /dev/null
+++ b/Engine/Source/Runtime/InteractiveToolsFramework/Public/BaseGizmos/GizmoElementShared.h
@@ -0,0 +1,62 @@
+// Copyright Epic Games, Inc. All Rights Reserved.
+
+#pragma once
+
+#include "CoreMinimal.h"
+#include "GizmoElementShared.generated.h"
+
+//
+// Visible/hittable state of gizmo element
+//
+UENUM()
+enum class EGizmoElementState : uint8
+{
+	None		= 0x00,
+	Visible		= 1<<1,
+	Hittable	= 1<<2,
+	VisibleAndHittable = Visible | Hittable
+};
+
+//
+// Interaction state of gizmo element
+//
+UENUM()
+enum class EGizmoElementInteractionState
+{
+	None,
+	Hovering,
+	Interacting
+};
+
+//
+// View dependent type: automatically cull gizmo element based on view.
+//
+//   Axis  - Cull object when angle between axis and view direction is within a given tolerance
+//   Plane - Cull object when angle between plane normal and view direction is perpendicular within a given tolerance
+//
+UENUM()
+enum class EGizmoElementViewDependentType
+{
+	None,
+	Axis,
+	Plane
+};
+
+//
+// View align type: automatically align gizmo element towards a view.
+// 
+//   PointOnly   - Align object forward axis to view direction only, useful for symmetrical objects such as a circle
+//   PointEye    - Align object forward axis to -camera view direction (camera pos - object center), align object up axis to scene view up
+//   PointScreen - Align object forward axis to scene view forward direction (view up ^ view right), align object up axis to scene view up
+//   Axial		- Rotate object around up axis, minimizing angle between forward axis and view direction
+//
+UENUM()
+enum class EGizmoElementViewAlignType
+{
+	None,
+	PointOnly,
+	PointEye,
+	PointScreen,
+	Axial
+};
+
