Added ray/ line-drawn circle intersection to gizmo element circle.

diff --git a/Engine/Source/Runtime/InteractiveToolsFramework/Private/BaseGizmos/GizmoElementCircle.cpp b/Engine/Source/Runtime/InteractiveToolsFramework/Private/BaseGizmos/GizmoElementCircle.cpp
--- a/Engine/Source/Runtime/InteractiveToolsFramework/Private/BaseGizmos/GizmoElementCircle.cpp
+++ b/Engine/Source/Runtime/InteractiveToolsFramework/Private/BaseGizmos/GizmoElementCircle.cpp
@@ -80,7 +80,31 @@ FInputRayHit UGizmoElementCircle::LineTrace(const FVector RayOrigin, const FVect
 		}
 		else if (bHitLine)
 		{
-			// @todo - add hit testing for line-drawn circle, requires storing PixelToWorld scale factor
+			FPlane Plane(WorldCenter, WorldNormal);
+			double HitDepth = FMath::RayPlaneIntersectionParam(RayOrigin, RayDirection, Plane);
+			if (HitDepth < 0)
+			{
+				return FInputRayHit();
+			}
+
+			FVector HitPoint = RayOrigin + RayDirection * HitDepth;
+
+			FVector NearestCirclePos;
+			GizmoMath::ClosetPointOnCircle(HitPoint, WorldCenter, WorldNormal, WorldRadius, NearestCirclePos);
+
+			FRay Ray(RayOrigin, RayDirection, true);
+			FVector NearestRayPos = Ray.ClosestPoint(NearestCirclePos);
+
+			const double HitBuffer = PixelHitThresholdAdjust + LineThickness;
+			double Distance = FVector::Distance(NearestCirclePos, NearestRayPos);
+			
+			if (Distance <= HitBuffer)
+			{
+				FInputRayHit RayHit(HitDepth);
+				RayHit.SetHitObject(this);
+				RayHit.HitIdentifier = PartIdentifier;
+				return RayHit;
+			}
 		}
 	}
 	return FInputRayHit();
