LineTrace now takes GizmoViewContext rather than using cached render state, which enables invisible hittable objects and will be more robust when there are multiple viewports.

diff --git a/Engine/Source/Editor/Experimental/EditorInteractiveToolsFramework/Private/EditorGizmos/EditorTransformGizmoBuilder.cpp b/Engine/Source/Editor/Experimental/EditorInteractiveToolsFramework/Private/EditorGizmos/EditorTransformGizmoBuilder.cpp
--- a/Engine/Source/Editor/Experimental/EditorInteractiveToolsFramework/Private/EditorGizmos/EditorTransformGizmoBuilder.cpp
+++ b/Engine/Source/Editor/Experimental/EditorInteractiveToolsFramework/Private/EditorGizmos/EditorTransformGizmoBuilder.cpp
@@ -5,6 +5,8 @@
 #include "EditorGizmos/EditorTransformGizmoSource.h"
 #include "EditorGizmos/EditorTransformProxy.h"
 #include "BaseGizmos/GizmoElementGroup.h"
+#include "BaseGizmos/GizmoViewContext.h"
+#include "ContextObjectStore.h"
 
 UInteractiveGizmo* UEditorTransformGizmoBuilder::BuildGizmo(const FToolBuilderState& SceneState) const
 {
@@ -12,10 +14,13 @@ UInteractiveGizmo* UEditorTransformGizmoBuilder::BuildGizmo(const FToolBuilderSt
 	TransformGizmo->Setup();
 	TransformGizmo->TransformGizmoSource = UEditorTransformGizmoSource::Construct(TransformGizmo);
 
+	UGizmoViewContext* GizmoViewContext = SceneState.ToolManager->GetContextObjectStore()->FindContext<UGizmoViewContext>();
+	check(GizmoViewContext && GizmoViewContext->IsValidLowLevel());
+
 	// @todo: Gizmo element construction to be moved here from UTransformGizmo.
 	// A UGizmoElementRenderMultiTarget will be constructed and both the
 	// render and hit target's Construct methods will take the gizmo element root as input.
-	TransformGizmo->HitTarget = UGizmoElementHitMultiTarget::Construct(TransformGizmo->GizmoElementRoot);
+	TransformGizmo->HitTarget = UGizmoElementHitMultiTarget::Construct(TransformGizmo->GizmoElementRoot, GizmoViewContext);
 
 	return TransformGizmo;
 }
@@ -27,5 +32,10 @@ void UEditorTransformGizmoBuilder::UpdateGizmoForSelection(UInteractiveGizmo* Gi
 		UEditorTransformProxy* TransformProxy = NewObject<UEditorTransformProxy>();
 		TransformGizmo->SetActiveTarget(TransformProxy);
 		TransformGizmo->SetVisibility(true);
+		
+		if (UGizmoElementHitMultiTarget* HitMultiTarget = Cast< UGizmoElementHitMultiTarget>(TransformGizmo->HitTarget))
+		{
+			HitMultiTarget->GizmoTransformProxy = TransformProxy;
+		}
 	}
 }
diff --git a/Engine/Source/Runtime/InteractiveToolsFramework/Private/BaseGizmos/GizmoElementArrow.cpp b/Engine/Source/Runtime/InteractiveToolsFramework/Private/BaseGizmos/GizmoElementArrow.cpp
--- a/Engine/Source/Runtime/InteractiveToolsFramework/Private/BaseGizmos/GizmoElementArrow.cpp
+++ b/Engine/Source/Runtime/InteractiveToolsFramework/Private/BaseGizmos/GizmoElementArrow.cpp
@@ -42,8 +42,7 @@ void UGizmoElementArrow::Render(IToolsContextRenderAPI* RenderAPI, const FRender
 	if (bVisibleViewDependent)
 	{
 		check(CylinderElement);
-		UGizmoElementBase* Element = Cast<UGizmoElementBase>(CylinderElement);
-		Element->Render(RenderAPI, CurrentRenderState);
+		CylinderElement->Render(RenderAPI, CurrentRenderState);
 
 		if (HeadType == EGizmoElementArrowHeadType::Cone)
 		{
@@ -58,27 +57,35 @@ void UGizmoElementArrow::Render(IToolsContextRenderAPI* RenderAPI, const FRender
 	}
 }
 
-FInputRayHit UGizmoElementArrow::LineTrace(const FVector RayOrigin, const FVector RayDirection)
+FInputRayHit UGizmoElementArrow::LineTrace(const UGizmoViewContext* ViewContext, const FLineTraceTraversalState& LineTraceState, const FVector& RayOrigin, const FVector& RayDirection)
 {
-	if (!IsHittableInView())
+	if (!IsHittable())
+	{
+		return FInputRayHit();
+	}
+
+	FLineTraceTraversalState CurrentLineTraceState(LineTraceState);
+	bool bHittableViewDependent = UpdateLineTraceState(ViewContext, Base, CurrentLineTraceState);
+
+	if (!bHittableViewDependent)
 	{
 		return FInputRayHit();
 	}
 
 	check(CylinderElement);
-	FInputRayHit Hit = CylinderElement->LineTrace(RayOrigin, RayDirection);
+	FInputRayHit Hit = CylinderElement->LineTrace(ViewContext, CurrentLineTraceState, RayOrigin, RayDirection);
 
 	if (!Hit.bHit)
 	{
 		if (HeadType == EGizmoElementArrowHeadType::Cone)
 		{
 			check(ConeElement);
-			Hit = ConeElement->LineTrace(RayOrigin, RayDirection);
+			Hit = ConeElement->LineTrace(ViewContext, CurrentLineTraceState, RayOrigin, RayDirection);
 		}
 		else // (HeadType == EGizmoElementArrowHeadType::Cube)
 		{
 			check(BoxElement);
-			Hit = BoxElement->LineTrace(RayOrigin, RayDirection);
+			Hit = BoxElement->LineTrace(ViewContext, CurrentLineTraceState, RayOrigin, RayDirection);
 		}
 	}
 
@@ -284,3 +291,5 @@ void UGizmoElementArrow::UpdateArrowHead()
 
 	bUpdateArrowHead = false;
 }
+
+#undef LOCTEXT_NAMESPACE
\ No newline at end of file
diff --git a/Engine/Source/Runtime/InteractiveToolsFramework/Private/BaseGizmos/GizmoElementBase.cpp b/Engine/Source/Runtime/InteractiveToolsFramework/Private/BaseGizmos/GizmoElementBase.cpp
--- a/Engine/Source/Runtime/InteractiveToolsFramework/Private/BaseGizmos/GizmoElementBase.cpp
+++ b/Engine/Source/Runtime/InteractiveToolsFramework/Private/BaseGizmos/GizmoElementBase.cpp
@@ -3,11 +3,75 @@
 #include "BaseGizmos/GizmoElementBase.h"
 #include "BaseGizmos/GizmoInterfaces.h"
 #include "BaseGizmos/GizmoRenderingUtil.h"
+#include "BaseGizmos/GizmoViewContext.h"
 #include "Materials/MaterialInterface.h"
 
 DEFINE_LOG_CATEGORY_STATIC(LogGizmoElementBase, Log, All);
 
+namespace GizmoElementBaseLocals
+{
+	template <class ViewType>
+	void GetViewInfo(const ViewType* View,
+		FVector& OutViewLocation, FVector& OutViewDirection, FVector& OutViewUp, bool& OutIsPerspectiveView)
+	{
+		OutIsPerspectiveView = View->IsPerspectiveProjection();
+		OutViewLocation = View->ViewLocation;
+		OutViewDirection = View->GetViewDirection();
+		OutViewUp = View->GetViewUp();
+	}
+}
+
 bool UGizmoElementBase::GetViewDependentVisibility(const FSceneView* View, const FTransform& InLocalToWorldTransform, const FVector& InLocalCenter) const
+{
+	bool bIsPerspectiveView;
+	FVector ViewLocation, ViewDirection, ViewUp;
+	GizmoElementBaseLocals::GetViewInfo<FSceneView>(View, ViewLocation, ViewDirection, ViewUp, bIsPerspectiveView);
+	return GetViewDependentVisibility(ViewLocation, ViewDirection, bIsPerspectiveView, InLocalToWorldTransform, InLocalCenter);
+}
+
+bool UGizmoElementBase::GetViewDependentVisibility(const UGizmoViewContext* View, const FTransform& InLocalToWorldTransform, const FVector& InLocalCenter) const
+{
+	bool bIsPerspectiveView;
+	FVector ViewLocation, ViewDirection, ViewUp;
+	GizmoElementBaseLocals::GetViewInfo<UGizmoViewContext>(View, ViewLocation, ViewDirection, ViewUp, bIsPerspectiveView);
+	return GetViewDependentVisibility(ViewLocation, ViewDirection, bIsPerspectiveView, InLocalToWorldTransform, InLocalCenter);
+}
+
+
+bool UGizmoElementBase::GetViewDependentHittable(const FSceneView* View, const FTransform& InLocalToWorldTransform, const FVector& InLocalCenter) const
+{
+	bool bIsPerspectiveView;
+	FVector ViewLocation, ViewDirection, ViewUp;
+	GizmoElementBaseLocals::GetViewInfo<FSceneView>(View, ViewLocation, ViewDirection, ViewUp, bIsPerspectiveView);
+	return GetViewDependentHittable(ViewLocation, ViewDirection, bIsPerspectiveView, InLocalToWorldTransform, InLocalCenter);
+}
+
+bool UGizmoElementBase::GetViewDependentHittable(const UGizmoViewContext* View, const FTransform& InLocalToWorldTransform, const FVector& InLocalCenter) const
+{
+	bool bIsPerspectiveView;
+	FVector ViewLocation, ViewDirection, ViewUp;
+	GizmoElementBaseLocals::GetViewInfo<UGizmoViewContext>(View, ViewLocation, ViewDirection, ViewUp, bIsPerspectiveView);
+	return GetViewDependentHittable(ViewLocation, ViewDirection, bIsPerspectiveView, InLocalToWorldTransform, InLocalCenter);
+}
+
+
+bool UGizmoElementBase::GetViewAlignRot(const FSceneView* View, const FTransform& InLocalToWorldTransform, const FVector& InLocalCenter, FQuat& OutAlignRot) const
+{
+	bool bIsPerspectiveView;
+	FVector ViewLocation, ViewDirection, ViewUp;
+	GizmoElementBaseLocals::GetViewInfo<FSceneView>(View, ViewLocation, ViewDirection, ViewUp, bIsPerspectiveView);
+	return GetViewAlignRot(ViewLocation, ViewDirection, ViewUp, bIsPerspectiveView, InLocalToWorldTransform, InLocalCenter, OutAlignRot);
+}
+
+bool UGizmoElementBase::GetViewAlignRot(const UGizmoViewContext* View, const FTransform& InLocalToWorldTransform, const FVector& InLocalCenter, FQuat& OutAlignRot) const
+{
+	bool bIsPerspectiveView;
+	FVector ViewLocation, ViewDirection, ViewUp;
+	GizmoElementBaseLocals::GetViewInfo<UGizmoViewContext>(View, ViewLocation, ViewDirection, ViewUp, bIsPerspectiveView);
+	return GetViewAlignRot(ViewLocation, ViewDirection, ViewUp, bIsPerspectiveView, InLocalToWorldTransform, InLocalCenter, OutAlignRot);
+}
+
+bool UGizmoElementBase::GetViewDependentVisibility(const FVector& InViewLocation, const FVector& InViewDirection, bool bInPerspectiveView, const FTransform& InLocalToWorldTransform, const FVector& InLocalCenter) const
 {
 	if (ViewDependentType == EGizmoElementViewDependentType::None || ViewAlignType == EGizmoElementViewAlignType::PointOnly || ViewAlignType == EGizmoElementViewAlignType::PointEye)
 	{
@@ -15,18 +79,16 @@ bool UGizmoElementBase::GetViewDependentVisibility(const FSceneView* View, const
 	}
 
 	FVector ViewDir;
-
-	if (View->IsPerspectiveProjection())
+	if (bInPerspectiveView)
 	{
-
 		FVector WorldCenter = InLocalToWorldTransform.TransformPosition(InLocalCenter);
-		ViewDir = WorldCenter - View->ViewLocation;
-		ViewDir.Normalize();
+		ViewDir = WorldCenter - InViewLocation;
 	}
 	else
 	{
-		ViewDir = View->GetViewDirection();
+		ViewDir = InViewDirection;
 	}
+	ViewDir.Normalize();
 
 	bool bVisibleViewDependent;
 	if (ViewDependentType == EGizmoElementViewDependentType::Axis)
@@ -42,25 +104,17 @@ bool UGizmoElementBase::GetViewDependentVisibility(const FSceneView* View, const
 
 }
 
-
-bool UGizmoElementBase::GetViewAlignRot(const FSceneView* View, const FTransform& InLocalToWorldTransform, const FVector& InLocalCenter, FQuat& OutAlignRot) const
+bool UGizmoElementBase::GetViewDependentHittable(const FVector& InViewLocation, const FVector& InViewDirection, bool bInPerspectiveView, const FTransform& InLocalToWorldTransform, const FVector& InLocalCenter) const
 {
-	if (ViewAlignType == EGizmoElementViewAlignType::None)
-	{
-		return false;
-	}
+	return (IsHittable() && (!IsVisible() || GetViewDependentVisibility(InViewLocation, InViewDirection, bInPerspectiveView, InLocalToWorldTransform, InLocalCenter)));
+}
 
+bool UGizmoElementBase::GetViewAlignRot(const FVector& InViewLocation, const FVector& InViewDirection, const FVector& InViewUp, bool bInPerspectiveView, const FTransform& InLocalToWorldTransform, const FVector& InLocalCenter, FQuat& OutAlignRot) const
+{
 	FVector Scale = InLocalToWorldTransform.GetScale3D();
-	if (!FMath::IsNearlyEqual(Scale.X, Scale.Y, KINDA_SMALL_NUMBER) || !FMath::IsNearlyEqual(Scale.X, Scale.Z, KINDA_SMALL_NUMBER))
+
+	if (ViewAlignType == EGizmoElementViewAlignType::None || !VerifyUniformScale(Scale))
 	{
-		// Warn that non-uniform scale is not currently supported 
-		bool bNonUniformScaleWarning = true;
-		if (bNonUniformScaleWarning)
-		{
-			UE_LOG(LogGizmoElementBase, Warning, TEXT("Gizmo element library view-dependent alignment does not currently support non-uniform scale (%f %f %f)."),
-				Scale.X, Scale.Y, Scale.Z);
-			bNonUniformScaleWarning = false;
-		}
 		return false;
 	}
 
@@ -68,15 +122,15 @@ bool UGizmoElementBase::GetViewAlignRot(const FSceneView* View, const FTransform
 
 	FTransform WorldToLocalTransform = InLocalToWorldTransform.Inverse();
 
-	if (View->IsPerspectiveProjection() && ViewAlignType != EGizmoElementViewAlignType::PointScreen)
+	if (bInPerspectiveView && ViewAlignType != EGizmoElementViewAlignType::PointScreen)
 	{
-		FVector LocalViewLocation = WorldToLocalTransform.TransformPosition(View->ViewLocation);
+		FVector LocalViewLocation = WorldToLocalTransform.TransformPosition(InViewLocation);
 		LocalViewDir = InLocalCenter - LocalViewLocation;
 		LocalViewDir.Normalize();
 	}
 	else
 	{
-		FVector WorldViewDir = View->GetViewDirection();
+		FVector WorldViewDir = InViewDirection;
 		LocalViewDir = WorldToLocalTransform.GetRotation().RotateVector(WorldViewDir);
 		LocalViewDir.Normalize();
 	}
@@ -91,7 +145,7 @@ bool UGizmoElementBase::GetViewAlignRot(const FSceneView* View, const FTransform
 		Right.Normalize();
 		FVector Up = ViewAlignNormal ^ Right;
 
-		FVector LocalViewUp = WorldToLocalTransform.TransformVector(View->GetViewUp());
+		FVector LocalViewUp = WorldToLocalTransform.TransformVector(InViewUp);
 		FVector TargetFwd = -LocalViewDir;
 		FVector TargetRight = LocalViewUp ^ TargetFwd;
 		TargetRight.Normalize();
@@ -117,6 +171,23 @@ bool UGizmoElementBase::GetViewAlignRot(const FSceneView* View, const FTransform
 	return true;
 }
 
+bool UGizmoElementBase::VerifyUniformScale(const FVector& Scale) const
+{
+	if (!FMath::IsNearlyEqual(Scale.X, Scale.Y, KINDA_SMALL_NUMBER) || !FMath::IsNearlyEqual(Scale.X, Scale.Z, KINDA_SMALL_NUMBER))
+	{
+		// Log one-time warning that non-uniform scale is not currently supported 
+		static bool bNonUniformScaleWarning = true;
+		if (bNonUniformScaleWarning)
+		{
+			UE_LOG(LogGizmoElementBase, Warning, TEXT("Gizmo element library view-dependent alignment does not currently support non-uniform scale (%f %f %f)."),
+				Scale.X, Scale.Y, Scale.Z);
+			bNonUniformScaleWarning = false;
+		}
+		return false;
+	}
+	return true;
+}
+
 FQuat UGizmoElementBase::GetAlignRotBetweenCoordSpaces(FVector SourceForward, FVector SourceRight, FVector SourceUp, FVector TargetForward, FVector TargetRight, FVector TargetUp) const
 {
 	FMatrix SourceToCanonical(
@@ -137,28 +208,22 @@ FQuat UGizmoElementBase::GetAlignRotBetweenCoordSpaces(FVector SourceForward, FV
 	return Result;
 }
 
-void UGizmoElementBase::CacheRenderState(const FTransform& InLocalToWorldState, double InPixelToWorldScale, bool InVisibleViewDependent)
+bool UGizmoElementBase::UpdateRenderState(IToolsContextRenderAPI* RenderAPI, const FVector& InLocalCenter, FRenderTraversalState& InOutRenderState)
 {
-	CachedLocalToWorldTransform = InLocalToWorldState;
-	CachedPixelToWorldScale = InPixelToWorldScale;
-	bHasCachedLocalToWorldTransform = true;
-	bCachedVisibleViewDependent = InVisibleViewDependent;
+	FQuat AlignRot;
+	bool bHasAlignRot;
+	return UpdateRenderState(RenderAPI, InLocalCenter, InOutRenderState, bHasAlignRot, AlignRot);
 }
 
-void UGizmoElementBase::ResetCachedRenderState()
-{
-	bHasCachedLocalToWorldTransform = false;
-	CachedLocalToWorldTransform = FTransform::Identity;
-	CachedPixelToWorldScale = 1.0f;
-	bCachedVisibleViewDependent = true;
-}
-
-bool UGizmoElementBase::UpdateRenderState(IToolsContextRenderAPI* RenderAPI, const FVector& InLocalOrigin, FRenderTraversalState& InOutRenderState)
+bool UGizmoElementBase::UpdateRenderState(IToolsContextRenderAPI* RenderAPI, const FVector & InLocalCenter, FRenderTraversalState & InOutRenderState, bool& bOutHasAlignRot, FQuat& OutAlignRot)
 {
 	check(RenderAPI);
 	const FSceneView* View = RenderAPI->GetSceneView();
 	check(View);
 
+	OutAlignRot = FQuat::Identity;
+	bOutHasAlignRot = false;
+
 	if (InOutRenderState.InteractionState == EGizmoElementInteractionState::None)
 	{
 		InOutRenderState.InteractionState = ElementInteractionState;
@@ -166,26 +231,42 @@ bool UGizmoElementBase::UpdateRenderState(IToolsContextRenderAPI* RenderAPI, con
 
 	InOutRenderState.MeshRenderState.Update(MeshRenderAttributes);
 
-	bool bVisibleViewDependent = GetViewDependentVisibility(View, InOutRenderState.LocalToWorldTransform, InLocalOrigin);
+	bool bVisibleViewDependent = GetViewDependentVisibility(View, InOutRenderState.LocalToWorldTransform, InLocalCenter);
 
 	if (bVisibleViewDependent)
 	{
-		FQuat AlignRot;
-		if (GetViewAlignRot(View, InOutRenderState.LocalToWorldTransform, InLocalOrigin, AlignRot))
-		{
-			InOutRenderState.LocalToWorldTransform = FTransform(AlignRot, InLocalOrigin) * InOutRenderState.LocalToWorldTransform;
-		}
-		else
-		{
-			InOutRenderState.LocalToWorldTransform = FTransform(InLocalOrigin) * InOutRenderState.LocalToWorldTransform;
-		}
+		bOutHasAlignRot = GetViewAlignRot(View, InOutRenderState.LocalToWorldTransform, InLocalCenter, OutAlignRot);
+		InOutRenderState.LocalToWorldTransform = FTransform(OutAlignRot, InLocalCenter) * InOutRenderState.LocalToWorldTransform;
 	}
 
-	CacheRenderState(InOutRenderState.LocalToWorldTransform, InOutRenderState.PixelToWorldScale, bVisibleViewDependent);
-
 	return bVisibleViewDependent;
 }
 
+bool UGizmoElementBase::UpdateLineTraceState(const UGizmoViewContext* ViewContext, const FVector& InLocalCenter, FLineTraceTraversalState& InOutRenderState)
+{
+	FQuat AlignQuat;
+	bool bHasAlignRot;
+	return UpdateLineTraceState(ViewContext, InLocalCenter, InOutRenderState, bHasAlignRot, AlignQuat);
+}
+
+bool UGizmoElementBase::UpdateLineTraceState(const UGizmoViewContext* ViewContext, const FVector& InLocalCenter, FLineTraceTraversalState& InOutRenderState, bool& bOutHasAlignRot, FQuat& OutAlignRot)
+{
+	check(ViewContext);
+
+	OutAlignRot = FQuat::Identity;
+	bOutHasAlignRot = false;
+
+	bool bHittableViewDependent = GetViewDependentHittable(ViewContext, InOutRenderState.LocalToWorldTransform, InLocalCenter);
+
+	if (bHittableViewDependent)
+	{
+		bOutHasAlignRot = GetViewAlignRot(ViewContext, InOutRenderState.LocalToWorldTransform, InLocalCenter, OutAlignRot);
+		InOutRenderState.LocalToWorldTransform = FTransform(OutAlignRot, InLocalCenter) * InOutRenderState.LocalToWorldTransform;
+	}
+
+	return bHittableViewDependent;
+}
+
 bool UGizmoElementBase::IsVisible() const
 {
 	bool bVisible = static_cast<uint8>(ElementState) & static_cast<uint8>(EGizmoElementState::Visible);
@@ -198,11 +279,6 @@ bool UGizmoElementBase::IsHittable() const
 	return (bEnabled && bHittable);
 }
 
-bool UGizmoElementBase::IsHittableInView() const
-{
-	return (IsHittable() && bHasCachedLocalToWorldTransform && (!IsVisible() || bCachedVisibleViewDependent));
-}
-
 void UGizmoElementBase::SetEnabled(bool InEnabled)
 {
 	bEnabled = InEnabled;
diff --git a/Engine/Source/Runtime/InteractiveToolsFramework/Private/BaseGizmos/GizmoElementBox.cpp b/Engine/Source/Runtime/InteractiveToolsFramework/Private/BaseGizmos/GizmoElementBox.cpp
--- a/Engine/Source/Runtime/InteractiveToolsFramework/Private/BaseGizmos/GizmoElementBox.cpp
+++ b/Engine/Source/Runtime/InteractiveToolsFramework/Private/BaseGizmos/GizmoElementBox.cpp
@@ -32,16 +32,25 @@ void UGizmoElementBox::Render(IToolsContextRenderAPI* RenderAPI, const FRenderTr
 	}
 }
 
-FInputRayHit UGizmoElementBox::LineTrace(const FVector RayOrigin, const FVector RayDirection)
+FInputRayHit UGizmoElementBox::LineTrace(const UGizmoViewContext* ViewContext, const FLineTraceTraversalState& LineTraceState, const FVector& RayOrigin, const FVector& RayDirection)
 {
-	if (IsHittableInView())
+	if (!IsHittable())
 	{
-		const FVector YAxis = CachedLocalToWorldTransform.TransformVectorNoScale(SideDirection);
-		const FVector ZAxis = CachedLocalToWorldTransform.TransformVectorNoScale(UpDirection);
+		return FInputRayHit();
+	}
+
+	FLineTraceTraversalState CurrentLineTraceState(LineTraceState);
+	bool bHittableViewDependent = UpdateLineTraceState(ViewContext, Center, CurrentLineTraceState);
+
+	if (bHittableViewDependent)
+	{
+		const FVector YAxis = CurrentLineTraceState.LocalToWorldTransform.TransformVectorNoScale(SideDirection);
+		const FVector ZAxis = CurrentLineTraceState.LocalToWorldTransform.TransformVectorNoScale(UpDirection);
 		const FVector XAxis = FVector::CrossProduct(YAxis, ZAxis);
-		const FVector WorldCenter = CachedLocalToWorldTransform.TransformPosition(FVector::ZeroVector);
-		const double Scale = CachedLocalToWorldTransform.GetScale3D().X;
-		const FVector WorldExtent = Dimensions * Scale * 0.5;
+		const FVector WorldCenter = CurrentLineTraceState.LocalToWorldTransform.TransformPosition(FVector::ZeroVector);
+		const double Scale = CurrentLineTraceState.LocalToWorldTransform.GetScale3D().X;
+		const double PixelHitThresholdAdjust = CurrentLineTraceState.PixelToWorldScale * PixelHitDistanceThreshold;
+		const FVector WorldExtent = Dimensions * Scale * 0.5 + FVector(PixelHitThresholdAdjust);
 
 		double HitDepth = 0.0;
 		UE::Geometry::TRay<double> Ray(RayOrigin, RayDirection);
diff --git a/Engine/Source/Runtime/InteractiveToolsFramework/Private/BaseGizmos/GizmoElementCircle.cpp b/Engine/Source/Runtime/InteractiveToolsFramework/Private/BaseGizmos/GizmoElementCircle.cpp
--- a/Engine/Source/Runtime/InteractiveToolsFramework/Private/BaseGizmos/GizmoElementCircle.cpp
+++ b/Engine/Source/Runtime/InteractiveToolsFramework/Private/BaseGizmos/GizmoElementCircle.cpp
@@ -48,14 +48,22 @@ void UGizmoElementCircle::Render(IToolsContextRenderAPI* RenderAPI, const FRende
 	}
 }
 
-FInputRayHit UGizmoElementCircle::LineTrace(const FVector RayOrigin, const FVector RayDirection)
+FInputRayHit UGizmoElementCircle::LineTrace(const UGizmoViewContext* ViewContext, const FLineTraceTraversalState& LineTraceState, const FVector& RayOrigin, const FVector& RayDirection)
 {
-	if (IsHittableInView())
+	if (!IsHittable())
 	{
-		const FVector WorldNormal = CachedLocalToWorldTransform.TransformVectorNoScale(Normal);
-		const FVector WorldCenter = CachedLocalToWorldTransform.TransformPosition(FVector::ZeroVector);
-		const double PixelHitThresholdAdjust = CachedPixelToWorldScale * PixelHitDistanceThreshold;
-		double WorldRadius = CachedLocalToWorldTransform.GetScale3D().X * Radius;
+		return FInputRayHit();
+	}
+
+	FLineTraceTraversalState CurrentLineTraceState(LineTraceState);
+	bool bHittableViewDependent = UpdateLineTraceState(ViewContext, Center, CurrentLineTraceState);
+
+	if (bHittableViewDependent)
+	{
+		const FVector WorldNormal = CurrentLineTraceState.LocalToWorldTransform.TransformVectorNoScale(Normal);
+		const FVector WorldCenter = CurrentLineTraceState.LocalToWorldTransform.TransformPosition(FVector::ZeroVector);
+		const double PixelHitThresholdAdjust = CurrentLineTraceState.PixelToWorldScale * PixelHitDistanceThreshold;
+		double WorldRadius = CurrentLineTraceState.LocalToWorldTransform.GetScale3D().X * Radius;
 
 		// if ray is parallel to circle, no hit
 		if (FMath::IsNearlyZero(FVector::DotProduct(WorldNormal, RayDirection)))
diff --git a/Engine/Source/Runtime/InteractiveToolsFramework/Private/BaseGizmos/GizmoElementCone.cpp b/Engine/Source/Runtime/InteractiveToolsFramework/Private/BaseGizmos/GizmoElementCone.cpp
--- a/Engine/Source/Runtime/InteractiveToolsFramework/Private/BaseGizmos/GizmoElementCone.cpp
+++ b/Engine/Source/Runtime/InteractiveToolsFramework/Private/BaseGizmos/GizmoElementCone.cpp
@@ -33,19 +33,27 @@ void UGizmoElementCone::Render(IToolsContextRenderAPI* RenderAPI, const FRenderT
 	}
 }
 
-FInputRayHit UGizmoElementCone::LineTrace(const FVector RayOrigin, const FVector RayDirection)
+FInputRayHit UGizmoElementCone::LineTrace(const UGizmoViewContext* ViewContext, const FLineTraceTraversalState& LineTraceState, const FVector& RayOrigin, const FVector& RayDirection)
 {
-	if (IsHittableInView())
+	if (!IsHittable())
+	{
+		return FInputRayHit();
+	}
+
+	FLineTraceTraversalState CurrentLineTraceState(LineTraceState);
+	bool bHittableViewDependent = UpdateLineTraceState(ViewContext, Origin, CurrentLineTraceState);
+
+	if (bHittableViewDependent)
 	{
 		bool bIntersects = false;
 		double RayParam = 0.0;
 
-		const double PixelHitThresholdAdjust = CachedPixelToWorldScale * PixelHitDistanceThreshold;
+		const double PixelHitThresholdAdjust = CurrentLineTraceState.PixelToWorldScale * PixelHitDistanceThreshold;
 		const double ConeSide = FMath::Sqrt(Height * Height + Radius * Radius);
 		const double CosAngle = Height / ConeSide;
-		const double WorldHeight = Height * CachedLocalToWorldTransform.GetScale3D().X + PixelHitThresholdAdjust * 2.0;
-		const FVector WorldDirection = CachedLocalToWorldTransform.TransformVectorNoScale(Direction);
-		const FVector WorldOrigin = CachedLocalToWorldTransform.TransformPosition(FVector::ZeroVector) - WorldDirection * PixelHitThresholdAdjust;
+		const double WorldHeight = Height * CurrentLineTraceState.LocalToWorldTransform.GetScale3D().X + PixelHitThresholdAdjust * 2.0;
+		const FVector WorldDirection = CurrentLineTraceState.LocalToWorldTransform.TransformVectorNoScale(Direction);
+		const FVector WorldOrigin = CurrentLineTraceState.LocalToWorldTransform.TransformPosition(FVector::ZeroVector) - WorldDirection * PixelHitThresholdAdjust;
 
 		GizmoMath::RayConeIntersection(
 			WorldOrigin,
diff --git a/Engine/Source/Runtime/InteractiveToolsFramework/Private/BaseGizmos/GizmoElementCylinder.cpp b/Engine/Source/Runtime/InteractiveToolsFramework/Private/BaseGizmos/GizmoElementCylinder.cpp
--- a/Engine/Source/Runtime/InteractiveToolsFramework/Private/BaseGizmos/GizmoElementCylinder.cpp
+++ b/Engine/Source/Runtime/InteractiveToolsFramework/Private/BaseGizmos/GizmoElementCylinder.cpp
@@ -35,20 +35,27 @@ void UGizmoElementCylinder::Render(IToolsContextRenderAPI* RenderAPI, const FRen
 }
 
 
-
-FInputRayHit UGizmoElementCylinder::LineTrace(const FVector RayOrigin, const FVector RayDirection)
+FInputRayHit UGizmoElementCylinder::LineTrace(const UGizmoViewContext* ViewContext, const FLineTraceTraversalState& LineTraceState, const FVector& RayOrigin, const FVector& RayDirection)
 {
-	if (IsHittableInView())
+	if (!IsHittable())
+	{
+		return FInputRayHit();
+	}
+
+	FLineTraceTraversalState CurrentLineTraceState(LineTraceState);
+	bool bHittableViewDependent = UpdateLineTraceState(ViewContext, Base, CurrentLineTraceState);
+
+	if (bHittableViewDependent)
 	{
 		bool bIntersects = false;
 		double RayParam = 0.0;
 		
-		const double PixelHitThresholdAdjust = CachedPixelToWorldScale * PixelHitDistanceThreshold;
-		const double WorldHeight = Height * CachedLocalToWorldTransform.GetScale3D().X + PixelHitThresholdAdjust * 2.0;
-		const double WorldRadius = Radius * CachedLocalToWorldTransform.GetScale3D().X + PixelHitThresholdAdjust;
-		const FVector WorldDirection = CachedLocalToWorldTransform.TransformVectorNoScale(Direction);
+		const double PixelHitThresholdAdjust = CurrentLineTraceState.PixelToWorldScale * PixelHitDistanceThreshold;
+		const double WorldHeight = Height * CurrentLineTraceState.LocalToWorldTransform.GetScale3D().X + PixelHitThresholdAdjust * 2.0;
+		const double WorldRadius = Radius * CurrentLineTraceState.LocalToWorldTransform.GetScale3D().X + PixelHitThresholdAdjust;
+		const FVector WorldDirection = CurrentLineTraceState.LocalToWorldTransform.TransformVectorNoScale(Direction);
 		const FVector LocalCenter = Direction * Height * 0.5;
-		const FVector WorldCenter = CachedLocalToWorldTransform.TransformPosition(LocalCenter);
+		const FVector WorldCenter = CurrentLineTraceState.LocalToWorldTransform.TransformPosition(LocalCenter);
 
 		GizmoMath::RayCylinderIntersection(
 			WorldCenter,
diff --git a/Engine/Source/Runtime/InteractiveToolsFramework/Private/BaseGizmos/GizmoElementGroup.cpp b/Engine/Source/Runtime/InteractiveToolsFramework/Private/BaseGizmos/GizmoElementGroup.cpp
--- a/Engine/Source/Runtime/InteractiveToolsFramework/Private/BaseGizmos/GizmoElementGroup.cpp
+++ b/Engine/Source/Runtime/InteractiveToolsFramework/Private/BaseGizmos/GizmoElementGroup.cpp
@@ -1,6 +1,18 @@
 // Copyright Epic Games, Inc. All Rights Reserved.
 
 #include "BaseGizmos/GizmoElementGroup.h"
+#include "BaseGizmos/GizmoViewContext.h"
+
+
+void UGizmoElementGroup::ApplyConstantScale(float PixelToWorldScale, FTransform& InOutLocalToWorldTransform)
+{
+	float Scale = InOutLocalToWorldTransform.GetScale3D().X;
+	if (bConstantScale)
+	{
+		Scale *= PixelToWorldScale;
+	}
+	InOutLocalToWorldTransform.SetScale3D(FVector(Scale, Scale, Scale));
+}
 
 void UGizmoElementGroup::Render(IToolsContextRenderAPI* RenderAPI, const FRenderTraversalState& RenderState)
 {
@@ -16,13 +28,7 @@ void UGizmoElementGroup::Render(IToolsContextRenderAPI* RenderAPI, const FRender
 
 	if (bVisibleViewDependent)
 	{
-		// Compute constant scale, if applicable
-		float Scale = CurrentRenderState.LocalToWorldTransform.GetScale3D().X;
-		if (bConstantScale)
-		{
-			Scale *= CurrentRenderState.PixelToWorldScale;
-		}
-		CurrentRenderState.LocalToWorldTransform.SetScale3D(FVector(Scale, Scale, Scale));
+		ApplyConstantScale(CurrentRenderState.PixelToWorldScale, CurrentRenderState.LocalToWorldTransform);
 
 		// Continue render even if not visible so all transforms will be cached 
 		// for subsequent line tracing.
@@ -36,17 +42,22 @@ void UGizmoElementGroup::Render(IToolsContextRenderAPI* RenderAPI, const FRender
 	}
 }
 
-FInputRayHit UGizmoElementGroup::LineTrace(const FVector Start, const FVector Direction)
+FInputRayHit UGizmoElementGroup::LineTrace(const UGizmoViewContext* ViewContext, const FLineTraceTraversalState& LineTraceState, const FVector& RayOrigin, const FVector& RayDirection)
 {
 	FInputRayHit Hit;
 
-	if (IsHittable())
+	FLineTraceTraversalState CurrentLineTraceState(LineTraceState);
+	bool bHittableViewDependent = UpdateLineTraceState(ViewContext, FVector::ZeroVector, CurrentLineTraceState);
+
+	if (bHittableViewDependent)
 	{
+		ApplyConstantScale(CurrentLineTraceState.PixelToWorldScale, CurrentLineTraceState.LocalToWorldTransform);
+
 		for (UGizmoElementBase* Element : Elements)
 		{
 			if (Element)
 			{
-				FInputRayHit NewHit = Element->LineTrace(Start, Direction);
+				FInputRayHit NewHit = Element->LineTrace(ViewContext, CurrentLineTraceState, RayOrigin, RayDirection);
 				if (!Hit.bHit || NewHit.HitDepth < Hit.HitDepth)
 				{
 					Hit = NewHit;
@@ -139,19 +150,6 @@ void UGizmoElementGroup::UpdatePartInteractionState(EGizmoElementInteractionStat
 	}
 }
 
-void UGizmoElementGroup::ResetCachedRenderState()
-{
-	Super::ResetCachedRenderState();
-
-	for (UGizmoElementBase* Element : Elements)
-	{
-		if (Element)
-		{
-			Element->ResetCachedRenderState();
-		}
-	}
-}
-
 void UGizmoElementGroup::SetConstantScale(bool bInConstantScale)
 {
 	bConstantScale = bInConstantScale;
diff --git a/Engine/Source/Runtime/InteractiveToolsFramework/Private/BaseGizmos/GizmoElementHitTargets.cpp b/Engine/Source/Runtime/InteractiveToolsFramework/Private/BaseGizmos/GizmoElementHitTargets.cpp
--- a/Engine/Source/Runtime/InteractiveToolsFramework/Private/BaseGizmos/GizmoElementHitTargets.cpp
+++ b/Engine/Source/Runtime/InteractiveToolsFramework/Private/BaseGizmos/GizmoElementHitTargets.cpp
@@ -2,14 +2,18 @@
 
 #include "BaseGizmos/GizmoElementHitTargets.h"
 #include "BaseGizmos/GizmoElementBase.h"
+#include "BaseGizmos/GizmoViewContext.h"
+#include "BaseGizmos/TransformProxy.h"
 #include "Engine/EngineTypes.h"    // FHitResult
 
 
 FInputRayHit UGizmoElementHitTarget::IsHit(const FInputDeviceRay& ClickPos) const
 {
-	if (GizmoElement && (!Condition || Condition(ClickPos)))
+	if (GizmoElement && GizmoViewContext && GizmoTransformProxy && (!Condition || Condition(ClickPos)))
 	{
-		return GizmoElement->LineTrace(ClickPos.WorldRay.Origin, ClickPos.WorldRay.Direction);
+		UGizmoElementBase::FLineTraceTraversalState LineTraceState;
+		LineTraceState.Initialize(GizmoViewContext, GizmoTransformProxy->GetTransform());
+		return GizmoElement->LineTrace(GizmoViewContext, LineTraceState, ClickPos.WorldRay.Origin, ClickPos.WorldRay.Direction);
 	}
 	return FInputRayHit();
 }
@@ -48,18 +52,21 @@ void UGizmoElementHitTarget::UpdateInteractingState(bool bInteracting)
 	}
 }
 
-UGizmoElementHitTarget* UGizmoElementHitTarget::Construct(UGizmoElementBase* InGizmoElement, UObject* Outer)
+UGizmoElementHitTarget* UGizmoElementHitTarget::Construct(UGizmoElementBase* InGizmoElement, UGizmoViewContext* InGizmoViewContext, UObject* Outer)
 {
 	UGizmoElementHitTarget* NewHitTarget = NewObject<UGizmoElementHitTarget>(Outer);
 	NewHitTarget->GizmoElement = InGizmoElement;
+	NewHitTarget->GizmoViewContext = InGizmoViewContext;
 	return NewHitTarget;
 }
 
 FInputRayHit UGizmoElementHitMultiTarget::IsHit(const FInputDeviceRay& ClickPos) const
 {
-	if (GizmoElement && (!Condition || Condition(ClickPos)))
+	if (GizmoElement && GizmoViewContext && GizmoTransformProxy && (!Condition || Condition(ClickPos)))
 	{
-		return GizmoElement->LineTrace(ClickPos.WorldRay.Origin, ClickPos.WorldRay.Direction);
+		UGizmoElementBase::FLineTraceTraversalState LineTraceState;
+		LineTraceState.Initialize(GizmoViewContext, GizmoTransformProxy->GetTransform());
+		return GizmoElement->LineTrace(GizmoViewContext, LineTraceState, ClickPos.WorldRay.Origin, ClickPos.WorldRay.Direction);
 	}
 	return FInputRayHit();
 }
@@ -114,10 +121,11 @@ void UGizmoElementHitMultiTarget::UpdateHittableState(bool bHittable, uint32 Par
 	GizmoElement->UpdatePartHittableState(bHittable, PartIdentifier);
 }
 
-UGizmoElementHitMultiTarget* UGizmoElementHitMultiTarget::Construct(UGizmoElementBase* InGizmoElement, UObject* Outer)
+UGizmoElementHitMultiTarget* UGizmoElementHitMultiTarget::Construct(UGizmoElementBase* InGizmoElement, UGizmoViewContext* InGizmoViewContext, UObject* Outer)
 {
 	UGizmoElementHitMultiTarget* NewHitTarget = NewObject<UGizmoElementHitMultiTarget>(Outer);
 	NewHitTarget->GizmoElement = InGizmoElement;
+	NewHitTarget->GizmoViewContext = InGizmoViewContext;
 	return NewHitTarget;
 }
 
diff --git a/Engine/Source/Runtime/InteractiveToolsFramework/Private/BaseGizmos/GizmoElementRectangle.cpp b/Engine/Source/Runtime/InteractiveToolsFramework/Private/BaseGizmos/GizmoElementRectangle.cpp
--- a/Engine/Source/Runtime/InteractiveToolsFramework/Private/BaseGizmos/GizmoElementRectangle.cpp
+++ b/Engine/Source/Runtime/InteractiveToolsFramework/Private/BaseGizmos/GizmoElementRectangle.cpp
@@ -46,16 +46,25 @@ void UGizmoElementRectangle::Render(IToolsContextRenderAPI* RenderAPI, const FRe
 	}
 }
 
-FInputRayHit UGizmoElementRectangle::LineTrace(const FVector RayOrigin, const FVector RayDirection)
+FInputRayHit UGizmoElementRectangle::LineTrace(const UGizmoViewContext* ViewContext, const FLineTraceTraversalState& LineTraceState, const FVector& RayOrigin, const FVector& RayDirection)
 {
-	if (IsHittableInView())
+	if (!IsHittable())
 	{
-		const FVector WorldUpAxis = CachedLocalToWorldTransform.TransformVectorNoScale(UpDirection);
-		const FVector WorldSideAxis = CachedLocalToWorldTransform.TransformVectorNoScale(SideDirection);
+		return FInputRayHit();
+	}
+
+	FLineTraceTraversalState CurrentLineTraceState(LineTraceState);
+	bool bHittableViewDependent = UpdateLineTraceState(ViewContext, Center, CurrentLineTraceState);
+
+	if (bHittableViewDependent)
+	{
+		FTransform LocalToWorldTransform = CurrentLineTraceState.LocalToWorldTransform;
+		const FVector WorldUpAxis = LocalToWorldTransform.TransformVectorNoScale(UpDirection);
+		const FVector WorldSideAxis = LocalToWorldTransform.TransformVectorNoScale(SideDirection);
 		const FVector WorldNormal = FVector::CrossProduct(WorldUpAxis, WorldSideAxis);
-		const FVector WorldCenter = CachedLocalToWorldTransform.TransformPosition(FVector::ZeroVector);
-		const double Scale = CachedLocalToWorldTransform.GetScale3D().X;
-		const double PixelHitThresholdAdjust = CachedPixelToWorldScale * PixelHitDistanceThreshold;
+		const FVector WorldCenter = LocalToWorldTransform.TransformPosition(FVector::ZeroVector);
+		const double Scale = LocalToWorldTransform.GetScale3D().X;
+		const double PixelHitThresholdAdjust = CurrentLineTraceState.PixelToWorldScale * PixelHitDistanceThreshold;
 		const double WorldHeight = Scale * Height + PixelHitThresholdAdjust * 2.0;
 		const double WorldWidth = Scale * Width + PixelHitThresholdAdjust * 2.0;
 		const FVector Base = WorldCenter - WorldUpAxis * WorldHeight * 0.5 - WorldSideAxis * WorldWidth * 0.5;
diff --git a/Engine/Source/Runtime/InteractiveToolsFramework/Private/BaseGizmos/GizmoElementTorus.cpp b/Engine/Source/Runtime/InteractiveToolsFramework/Private/BaseGizmos/GizmoElementTorus.cpp
--- a/Engine/Source/Runtime/InteractiveToolsFramework/Private/BaseGizmos/GizmoElementTorus.cpp
+++ b/Engine/Source/Runtime/InteractiveToolsFramework/Private/BaseGizmos/GizmoElementTorus.cpp
@@ -36,7 +36,7 @@ void UGizmoElementTorus::Render(IToolsContextRenderAPI* RenderAPI, const FRender
 
 
 
-FInputRayHit UGizmoElementTorus::LineTrace(const FVector RayOrigin, const FVector RayDirection)
+FInputRayHit UGizmoElementTorus::LineTrace(const UGizmoViewContext* ViewContext, const FLineTraceTraversalState& LineTraceState, const FVector& RayOrigin, const FVector& RayDirection)
 {
 	// Line trace is not supported for torus.
 	return FInputRayHit();
@@ -160,3 +160,4 @@ bool UGizmoElementTorus::GetEndCaps() const
 	return bEndCaps;
 }
 
+#undef LOCTEXT_NAMESPACE
\ No newline at end of file
diff --git a/Engine/Source/Runtime/InteractiveToolsFramework/Public/BaseGizmos/GizmoElementArrow.h b/Engine/Source/Runtime/InteractiveToolsFramework/Public/BaseGizmos/GizmoElementArrow.h
--- a/Engine/Source/Runtime/InteractiveToolsFramework/Public/BaseGizmos/GizmoElementArrow.h
+++ b/Engine/Source/Runtime/InteractiveToolsFramework/Public/BaseGizmos/GizmoElementArrow.h
@@ -35,7 +35,7 @@ public:
 
 	//~ Begin UGizmoElementBase Interface.
 	virtual void Render(IToolsContextRenderAPI* RenderAPI, const FRenderTraversalState& RenderState) override;
-	virtual FInputRayHit LineTrace(const FVector Start, const FVector Direction) override;
+	virtual FInputRayHit LineTrace(const UGizmoViewContext* ViewContext, const FLineTraceTraversalState& LineTraceState, const FVector& RayOrigin, const FVector& RayDirection) override;
 	virtual FBoxSphereBounds CalcBounds(const FTransform& LocalToWorld) const override;
 	//~ End UGizmoElementBase Interface.
 
diff --git a/Engine/Source/Runtime/InteractiveToolsFramework/Public/BaseGizmos/GizmoElementBase.h b/Engine/Source/Runtime/InteractiveToolsFramework/Public/BaseGizmos/GizmoElementBase.h
--- a/Engine/Source/Runtime/InteractiveToolsFramework/Public/BaseGizmos/GizmoElementBase.h
+++ b/Engine/Source/Runtime/InteractiveToolsFramework/Public/BaseGizmos/GizmoElementBase.h
@@ -14,6 +14,7 @@
 #include "GizmoElementBase.generated.h"
 
 class UMaterialInterface;
+class UGizmoViewContext;
 
 /**
  * Base class for 2d and 3d primitive objects intended to be used as part of 3D Gizmos.
@@ -88,29 +89,40 @@ public:
 		}
 	};
 
+	struct FLineTraceTraversalState
+	{
+		// LocalToWorld transform 
+		// Note: non-uniform scale is not supported and the X scale element will be used for uniform scaling.
+		FTransform LocalToWorldTransform;
+
+		// Pixel to world scale
+		double PixelToWorldScale = 1.0;
+
+		// Initialize state 
+		void Initialize(const UGizmoViewContext* InGizmoViewContext, FTransform InTransform)
+		{
+			LocalToWorldTransform = InTransform;
+			PixelToWorldScale = GizmoRenderingUtil::CalculateLocalPixelToWorldScale(InGizmoViewContext, InTransform.GetLocation());
+		}
+	};
+
 public:
 
 	// Render enabled visible element.
 	virtual void Render(IToolsContextRenderAPI* RenderAPI, const FRenderTraversalState& RenderState) PURE_VIRTUAL(UGizmoElementBase::Render);
 
 	// Line trace enabled hittable element.
-	virtual FInputRayHit LineTrace(const FVector Start, const FVector Direction) PURE_VIRTUAL(UGizmoElementBase::LineTrace, return FInputRayHit(););
+	virtual FInputRayHit LineTrace(const UGizmoViewContext* ViewContext, const FLineTraceTraversalState& LineTraceState, const FVector& RayOrigin, const FVector& RayDirection) PURE_VIRTUAL(UGizmoElementBase::LineTrace, return FInputRayHit(););
 
 	// Calcute box sphere bounds for use when hit testing.
 	virtual FBoxSphereBounds CalcBounds(const FTransform& LocalToWorld) const PURE_VIRTUAL(UGizmoElementBase::CalcBounds, return FBoxSphereBounds(););
 
-	// Reset the cached render state.
-	virtual void ResetCachedRenderState();
-
 	// Whether this object is visible.
 	virtual bool IsVisible() const;
 
 	// Whether this object is hittable.
 	virtual bool IsHittable() const;
 
-	// Whether this object is hittable in the most recently cached view.
-	virtual bool IsHittableInView() const;
-
 	// Render and LineTrace should only occur when bEnabled is true.
 	virtual void SetEnabled(bool InEnabled);
 	virtual bool GetEnabled() const;
@@ -327,48 +339,72 @@ protected:
 	UPROPERTY()
 	float PixelHitDistanceThreshold = 7.0;
 
-	// Render stores the last cached transform, used by line trace.
-	UPROPERTY()
-	FTransform CachedLocalToWorldTransform = FTransform::Identity;
-
-	// Render stores the last pixel to world scale, used by line trace.
-	UPROPERTY()
-	float CachedPixelToWorldScale = 1.0;
+protected:
 
-	// Whether last transform has been cached.
-	UPROPERTY()
-	bool bHasCachedLocalToWorldTransform = false;
+	// Returns whether object is visible in input FSceneView based on view-dependent visibility settings 
+	virtual bool GetViewDependentVisibility(const FSceneView* View, const FTransform& InLocalToWorldTransform, const FVector& InLocalCenter) const;
 
-	// Whether visible object was visible during the last render.
-	UPROPERTY()
-	bool bCachedVisibleViewDependent = true;
+	// Returns whether object is visible in input gizmo view context based on view-dependent visibility settings
+	virtual bool GetViewDependentVisibility(const UGizmoViewContext* View, const FTransform& InLocalToWorldTransform, const FVector& InLocalCenter) const;
 
-	// Cached box sphere bounds.
-	UPROPERTY()
-	FBoxSphereBounds CachedBoxSphereBounds;
+	// Returns whether object is visible based input view parameters and view-dependent visibility settings
+	virtual bool GetViewDependentVisibility(const FVector& InViewLocation, const FVector& InViewDirection, bool bInPerspectiveView, const FTransform& InLocalToWorldTransform, const FVector& InLocalCenter) const;
 
-	// Cached box sphere bounds.
-	UPROPERTY()
-	bool bHasCachedBoxSphereBounds;
+	// Returns whether object is hittable in input FSceneView based on element state and view-dependent settings. 
+	virtual bool GetViewDependentHittable(const FSceneView* View, const FTransform& InLocalToWorldTransform, const FVector& InLocalCenter) const;
 
-protected:
+	// Returns whether object is hittable in input gizmo view context based on element state and view-dependent settings. 
+	virtual bool GetViewDependentHittable(const UGizmoViewContext* View, const FTransform& InLocalToWorldTransform, const FVector& InLocalCenter) const;
 
-	// Returns whether object is visible based on view-dependent visibility settings. 
-	virtual bool GetViewDependentVisibility(const FSceneView* View, const FTransform& InLocalToWorldTransform, const FVector& InLocalCenter) const;
+	// Returns whether object is hittable based on input view parameters, element state and view-dependent settings. 
+	virtual bool GetViewDependentHittable(const FVector& InViewLocation, const FVector& InViewDirection, bool bInPerspectiveView, const FTransform& InLocalToWorldTransform, const FVector& InLocalCenter) const;
 
-	// Returns true when view alignment is enabled. OutAlignRot is the rotation in local space which will align the object to the view base
-	// on view-dependent alignment settings. So it should be prepended to the local-to-world transform.
+	// Return whether this element has a view alignment rotation based on input FSceneView and view-dependent alignment settings.
+	// @param OutAlignRot the rotation to align this element in local space, should be prepended to the local-to-world transform.
 	virtual bool GetViewAlignRot(const FSceneView* View, const FTransform& InLocalToWorldTransform, const FVector& InLocalCenter, FQuat& OutAlignRot) const;
 
-	// Helper method to calculate rotation between coord spaces.
-	FQuat GetAlignRotBetweenCoordSpaces(FVector SourceForward, FVector SourceSide, FVector SourceUp, FVector TargetForward, FVector TargetSide, FVector TargetUp) const;
+	// Return whether this element has a view alignment rotation based on input gizmo view context and view-dependent alignment settings.
+	// @param OutAlignRot the rotation to align this element in local space, should be prepended to the local-to-world transform.
+	virtual bool GetViewAlignRot(const UGizmoViewContext* View, const FTransform& InLocalToWorldTransform, const FVector& InLocalCenter, FQuat& OutAlignRot) const;
+
+	// Return whether this element has a view alignment rotation based on input view parameters and view-dependent alignment settings.
+	// @param OutAlignRot the rotation to align this element in local space, should be prepended to the local-to-world transform.
+	virtual bool GetViewAlignRot(const FVector& InViewLocation, const FVector& InViewDirection, const FVector& InViewUp, bool bInPerspectiveView, const FTransform& InLocalToWorldTransform, const FVector& InLocalCenter, FQuat& OutAlignRot) const;
 
 	// Update render state during render traversal, determines the current render state for this element 
+	// @param RenderAPI - tools render context
+	// @param InLocalCenter - local element center position
+	// @param InOutRenderState - render state's local to world transform will be updated with translation to center and view-dependent alignment rotation, if applicable.
 	// @return view dependent visibility, true if this element is visible in the current view. 
-	virtual bool UpdateRenderState(IToolsContextRenderAPI* RenderAPI, const FVector& InLocalOrigin, FRenderTraversalState& InOutRenderState);
+	virtual bool UpdateRenderState(IToolsContextRenderAPI* RenderAPI, const FVector& InLocalCenter, FRenderTraversalState& InOutRenderState);
+
+	// Update render state during render traversal, determines the current render state for this element, same parameters as
+	// Same parameters as UpdateRenderState above plus two output parameters:
+	// @param bOutHasAlignRot - whether alignment rotation was applied to output state
+	// @param OutAlignRot - alignment rotation applied to output state, if applicable.
+	virtual bool UpdateRenderState(IToolsContextRenderAPI* RenderAPI, const FVector& InLocalCenter, FRenderTraversalState& InOutRenderState, bool& bOutHasAlignRot, FQuat& OutAlignRot);
+
+	// Update line trace state during line trace traversal, determines the current state for this element 
+	// @param ViewContext - current gizmo view context
+	// @param InLocalCenter - local element center position, this will update InOutLineTraceState's LocalToWorldTransform with a translation to LocalCenter
+	// @param InOutLineTraceState - line trace state's local to world transform will be updated with translation to center and view-dependent alignment rotation, if applicable.
+	// @return view dependent visibility, true if this element is visible in the current view. 
+	virtual bool UpdateLineTraceState(const UGizmoViewContext* ViewContext, const FVector& InLocalCenter, FLineTraceTraversalState& InOutLineTraceState);
+
+	// Update line trace state during line trace traversal, determines the current state for this element 
+	// Same parameters as UpdateLineTraceState above plus two output parameters:
+	// @param bOutHasAlignRot - whether alignment rotation was applied to output state
+	// @param OutAlignRot - alignment rotation applied to output state, if applicable.
+	virtual bool UpdateLineTraceState(const UGizmoViewContext* ViewContext, const FVector& InLocalCenter, FLineTraceTraversalState& InOutRenderState, bool& bOutHasAlignRot, FQuat& OutAlignRot);
+
+protected:
 
-	// Cache render state during render traversal to be used subsequently when line tracing.
-	virtual void CacheRenderState(const FTransform& InLocalToWorldState, double InPixelToWorldScale, bool InVisibleViewDependent = true);
+	// Helper method to verify scale is uniform. If it is non-uniform, a one-time warning log is issued.
+	// Returns true if scale is uniform.
+	bool VerifyUniformScale(const FVector& Scale) const;
 
+	// Helper method for view alignment.
+	// Returns rotation between source and target input coordinate spaces.
+	FQuat GetAlignRotBetweenCoordSpaces(FVector SourceForward, FVector SourceRight, FVector SourceUp, FVector TargetForward, FVector TargetRight, FVector TargetUp) const;
 };
 
diff --git a/Engine/Source/Runtime/InteractiveToolsFramework/Public/BaseGizmos/GizmoElementBox.h b/Engine/Source/Runtime/InteractiveToolsFramework/Public/BaseGizmos/GizmoElementBox.h
--- a/Engine/Source/Runtime/InteractiveToolsFramework/Public/BaseGizmos/GizmoElementBox.h
+++ b/Engine/Source/Runtime/InteractiveToolsFramework/Public/BaseGizmos/GizmoElementBox.h
@@ -22,7 +22,7 @@ public:
 
 	//~ Begin UGizmoElementBase Interface.
 	virtual void Render(IToolsContextRenderAPI* RenderAPI, const FRenderTraversalState& RenderState) override;
-	virtual FInputRayHit LineTrace(const FVector Start, const FVector Direction) override;
+	virtual FInputRayHit LineTrace(const UGizmoViewContext* ViewContext, const FLineTraceTraversalState& LineTraceState, const FVector& RayOrigin, const FVector& RayDirection) override;
 	virtual FBoxSphereBounds CalcBounds(const FTransform& LocalToWorld) const override;
 	//~ End UGizmoElementBase Interface.
 
diff --git a/Engine/Source/Runtime/InteractiveToolsFramework/Public/BaseGizmos/GizmoElementCircle.h b/Engine/Source/Runtime/InteractiveToolsFramework/Public/BaseGizmos/GizmoElementCircle.h
--- a/Engine/Source/Runtime/InteractiveToolsFramework/Public/BaseGizmos/GizmoElementCircle.h
+++ b/Engine/Source/Runtime/InteractiveToolsFramework/Public/BaseGizmos/GizmoElementCircle.h
@@ -20,7 +20,7 @@ class INTERACTIVETOOLSFRAMEWORK_API UGizmoElementCircle : public UGizmoElementLi
 public:
 	//~ Begin UGizmoElementBase Interface.
 	virtual void Render(IToolsContextRenderAPI* RenderAPI, const FRenderTraversalState& RenderState) override;
-	virtual FInputRayHit LineTrace(const FVector Start, const FVector Direction) override;
+	virtual FInputRayHit LineTrace(const UGizmoViewContext* ViewContext, const FLineTraceTraversalState& LineTraceState, const FVector& RayOrigin, const FVector& RayDirection) override;
 	virtual FBoxSphereBounds CalcBounds(const FTransform& LocalToWorld) const override;
 	//~ End UGizmoElementBase Interface.
 
diff --git a/Engine/Source/Runtime/InteractiveToolsFramework/Public/BaseGizmos/GizmoElementCone.h b/Engine/Source/Runtime/InteractiveToolsFramework/Public/BaseGizmos/GizmoElementCone.h
--- a/Engine/Source/Runtime/InteractiveToolsFramework/Public/BaseGizmos/GizmoElementCone.h
+++ b/Engine/Source/Runtime/InteractiveToolsFramework/Public/BaseGizmos/GizmoElementCone.h
@@ -22,7 +22,7 @@ public:
 
 	//~ Begin UGizmoElementBase Interface.
 	virtual void Render(IToolsContextRenderAPI* RenderAPI, const FRenderTraversalState& RenderState) override;
-	virtual FInputRayHit LineTrace(const FVector Start, const FVector Direction) override;
+	virtual FInputRayHit LineTrace(const UGizmoViewContext* ViewContext, const FLineTraceTraversalState& LineTraceState, const FVector& RayOrigin, const FVector& RayDirection) override;
 	virtual FBoxSphereBounds CalcBounds(const FTransform& LocalToWorld) const override;
 	//~ End UGizmoElementBase Interface.
 
diff --git a/Engine/Source/Runtime/InteractiveToolsFramework/Public/BaseGizmos/GizmoElementCylinder.h b/Engine/Source/Runtime/InteractiveToolsFramework/Public/BaseGizmos/GizmoElementCylinder.h
--- a/Engine/Source/Runtime/InteractiveToolsFramework/Public/BaseGizmos/GizmoElementCylinder.h
+++ b/Engine/Source/Runtime/InteractiveToolsFramework/Public/BaseGizmos/GizmoElementCylinder.h
@@ -22,7 +22,7 @@ public:
 
 	//~ Begin UGizmoElementBase Interface.
 	virtual void Render(IToolsContextRenderAPI* RenderAPI, const FRenderTraversalState& RenderState) override;
-	virtual FInputRayHit LineTrace(const FVector Start, const FVector Direction) override;
+	virtual FInputRayHit LineTrace(const UGizmoViewContext* ViewContext, const FLineTraceTraversalState& LineTraceState, const FVector& RayOrigin, const FVector& RayDirection) override;
 	virtual FBoxSphereBounds CalcBounds(const FTransform& LocalToWorld) const override;
 	//~ End UGizmoElementBase Interface.
 
diff --git a/Engine/Source/Runtime/InteractiveToolsFramework/Public/BaseGizmos/GizmoElementGroup.h b/Engine/Source/Runtime/InteractiveToolsFramework/Public/BaseGizmos/GizmoElementGroup.h
--- a/Engine/Source/Runtime/InteractiveToolsFramework/Public/BaseGizmos/GizmoElementGroup.h
+++ b/Engine/Source/Runtime/InteractiveToolsFramework/Public/BaseGizmos/GizmoElementGroup.h
@@ -21,7 +21,7 @@ public:
 
 	//~ Begin UGizmoElementBase Interface.
 	virtual void Render(IToolsContextRenderAPI* RenderAPI, const FRenderTraversalState& RenderState) override;
-	virtual FInputRayHit LineTrace(const FVector Start, const FVector Direction) override;
+	virtual FInputRayHit LineTrace(const UGizmoViewContext* ViewContext, const FLineTraceTraversalState& LineTraceState, const FVector& RayOrigin, const FVector& RayDirection) override;
 	virtual FBoxSphereBounds CalcBounds(const FTransform& LocalToWorld) const override;
 	//~ End UGizmoElementBase Interface.
 
@@ -40,9 +40,6 @@ public:
 	// Update group and contained elements' interaction state for elements in specified gizmo parts.
 	virtual void UpdatePartInteractionState(EGizmoElementInteractionState InInteractionState, uint32 InPartIdentifier);
 
-	// Reset cached render state
-	virtual void ResetCachedRenderState();
-
 	// When true, maintains view-dependent constant scale for this gizmo object hierarchy
 	virtual void SetConstantScale(bool InConstantScale);
 	virtual bool GetConstantScale() const;
@@ -61,4 +58,7 @@ protected:
 	// Gizmo elements within this group
 	UPROPERTY()
 	TArray<TObjectPtr<UGizmoElementBase>> Elements;
+
+	// Applies constant scale to the current local to world transform, if constant scale is true
+	virtual void ApplyConstantScale(float PixelToWorldScale, FTransform& InOutLocalToWorldTransform);
 };
diff --git a/Engine/Source/Runtime/InteractiveToolsFramework/Public/BaseGizmos/GizmoElementHitTargets.h b/Engine/Source/Runtime/InteractiveToolsFramework/Public/BaseGizmos/GizmoElementHitTargets.h
--- a/Engine/Source/Runtime/InteractiveToolsFramework/Public/BaseGizmos/GizmoElementHitTargets.h
+++ b/Engine/Source/Runtime/InteractiveToolsFramework/Public/BaseGizmos/GizmoElementHitTargets.h
@@ -4,7 +4,9 @@
 
 #include "CoreMinimal.h"
 #include "BaseGizmos/GizmoInterfaces.h"
+#include "BaseGizmos/GizmoViewContext.h"
 #include "BaseGizmos/HitTargets.h"
+#include "BaseGizmos/TransformProxy.h"
 #include "GizmoElementHitTargets.generated.h"
 
 class UGizmoElementBase;
@@ -27,6 +29,12 @@ public:
 	UPROPERTY()
 	TObjectPtr<UGizmoElementBase> GizmoElement;
 
+	UPROPERTY()
+	TObjectPtr<UGizmoViewContext> GizmoViewContext;
+
+	UPROPERTY()
+	TObjectPtr<UTransformProxy> GizmoTransformProxy;
+
 	/**
 	 * If set, this condition is checked before performing the hit test. This gives a way
 	 * to disable the hit test without hiding the component. This is useful, for instance,
@@ -45,6 +53,7 @@ public:
 public:
 	static UGizmoElementHitTarget* Construct(
 		UGizmoElementBase* InGizmoElement,
+		UGizmoViewContext* InGizmoViewContext,
 		UObject* Outer = (UObject*)GetTransientPackage());
 };
 
@@ -70,6 +79,12 @@ public:
 	UPROPERTY()
 	TObjectPtr<UGizmoElementBase> GizmoElement;
 
+	UPROPERTY()
+	TObjectPtr<UGizmoViewContext> GizmoViewContext;
+
+	UPROPERTY()
+	TObjectPtr<UTransformProxy> GizmoTransformProxy;
+
 	/**
 	 * If set, this condition is checked before performing the hit test. This gives a way
 	 * to disable the hit test without hiding the component. This is useful, for instance,
@@ -98,6 +113,7 @@ protected:
 public:
 	static UGizmoElementHitMultiTarget* Construct(
 		UGizmoElementBase* InGizmoElement,
+		UGizmoViewContext* InGizmoViewContext,
 		UObject* Outer = (UObject*)GetTransientPackage());
 
 };
diff --git a/Engine/Source/Runtime/InteractiveToolsFramework/Public/BaseGizmos/GizmoElementRectangle.h b/Engine/Source/Runtime/InteractiveToolsFramework/Public/BaseGizmos/GizmoElementRectangle.h
--- a/Engine/Source/Runtime/InteractiveToolsFramework/Public/BaseGizmos/GizmoElementRectangle.h
+++ b/Engine/Source/Runtime/InteractiveToolsFramework/Public/BaseGizmos/GizmoElementRectangle.h
@@ -24,7 +24,7 @@ public:
 
 	//~ Begin UGizmoElementBase Interface.
 	virtual void Render(IToolsContextRenderAPI* RenderAPI, const FRenderTraversalState& RenderState) override;
-	virtual FInputRayHit LineTrace(const FVector Start, const FVector Direction) override;
+	virtual FInputRayHit LineTrace(const UGizmoViewContext* ViewContext, const FLineTraceTraversalState& LineTraceState, const FVector& RayOrigin, const FVector& RayDirection) override;
 	virtual FBoxSphereBounds CalcBounds(const FTransform& LocalToWorld) const override;
 	//~ End UGizmoElementBase Interface.
 
diff --git a/Engine/Source/Runtime/InteractiveToolsFramework/Public/BaseGizmos/GizmoElementTorus.h b/Engine/Source/Runtime/InteractiveToolsFramework/Public/BaseGizmos/GizmoElementTorus.h
--- a/Engine/Source/Runtime/InteractiveToolsFramework/Public/BaseGizmos/GizmoElementTorus.h
+++ b/Engine/Source/Runtime/InteractiveToolsFramework/Public/BaseGizmos/GizmoElementTorus.h
@@ -27,7 +27,7 @@ public:
 
 	//~ Begin UGizmoElementBase Interface.
 	virtual void Render(IToolsContextRenderAPI* RenderAPI, const FRenderTraversalState& RenderState) override;
-	virtual FInputRayHit LineTrace(const FVector Start, const FVector Direction) override;
+	virtual FInputRayHit LineTrace(const UGizmoViewContext* ViewContext, const FLineTraceTraversalState& LineTraceState, const FVector& RayOrigin, const FVector& RayDirection) override;
 	virtual FBoxSphereBounds CalcBounds(const FTransform& LocalToWorld) const override;
 	//~ End UGizmoElementBase Interface.
 
