Remove irrelevant undefs.

diff --git a/Engine/Source/Runtime/InteractiveToolsFramework/Private/BaseGizmos/GizmoElementArrow.cpp b/Engine/Source/Runtime/InteractiveToolsFramework/Private/BaseGizmos/GizmoElementArrow.cpp
--- a/Engine/Source/Runtime/InteractiveToolsFramework/Private/BaseGizmos/GizmoElementArrow.cpp
+++ b/Engine/Source/Runtime/InteractiveToolsFramework/Private/BaseGizmos/GizmoElementArrow.cpp
@@ -291,5 +291,3 @@ void UGizmoElementArrow::UpdateArrowHead()
 
 	bUpdateArrowHead = false;
 }
-
-#undef LOCTEXT_NAMESPACE
\ No newline at end of file
diff --git a/Engine/Source/Runtime/InteractiveToolsFramework/Private/BaseGizmos/GizmoElementTorus.cpp b/Engine/Source/Runtime/InteractiveToolsFramework/Private/BaseGizmos/GizmoElementTorus.cpp
--- a/Engine/Source/Runtime/InteractiveToolsFramework/Private/BaseGizmos/GizmoElementTorus.cpp
+++ b/Engine/Source/Runtime/InteractiveToolsFramework/Private/BaseGizmos/GizmoElementTorus.cpp
@@ -159,5 +159,3 @@ bool UGizmoElementTorus::GetEndCaps() const
 {
 	return bEndCaps;
 }
-
-#undef LOCTEXT_NAMESPACE
\ No newline at end of file
