ITF Context Object Stores search through parent stores in all Find implementations

diff --git a/Engine/Source/Runtime/InteractiveToolsFramework/Private/ContextObjectStore.cpp b/Engine/Source/Runtime/InteractiveToolsFramework/Private/ContextObjectStore.cpp
--- a/Engine/Source/Runtime/InteractiveToolsFramework/Private/ContextObjectStore.cpp
+++ b/Engine/Source/Runtime/InteractiveToolsFramework/Private/ContextObjectStore.cpp
@@ -11,6 +11,11 @@ UObject* UContextObjectStore::FindContextByClass(UClass* InClass) const
 			return ContextObject;
 		}
 	}
+	
+	if (UContextObjectStore* ParentStore = Cast<UContextObjectStore>(GetOuter()))
+	{
+		return ParentStore->FindContextByClass(InClass);
+	}
 
 	return nullptr;
 }
