ToolsFramework: Add UInteractiveToolManager::PostActiveToolShutdownRequest(). This function allows a Tool to request that it be shut down, for example a Tool that wants to apply an action and exit on-click.

Also added delegate UInteractiveToolManager::OnToolShutdownRequest. If bound, PostActiveToolShutdownRequest() will call this delegate first, to allow higher-level code to handle the shutdown request (eg by forwarding to EdModeInteractiveToolsContext::EndTool(), for example)

diff --git a/Engine/Source/Runtime/InteractiveToolsFramework/Private/InteractiveToolManager.cpp b/Engine/Source/Runtime/InteractiveToolsFramework/Private/InteractiveToolManager.cpp
--- a/Engine/Source/Runtime/InteractiveToolsFramework/Private/InteractiveToolManager.cpp
+++ b/Engine/Source/Runtime/InteractiveToolsFramework/Private/InteractiveToolManager.cpp
@@ -434,6 +434,27 @@ bool UInteractiveToolManager::RequestSelectionChange(const FSelectedOjectsChange
 }
 
 
+bool UInteractiveToolManager::PostActiveToolShutdownRequest(UInteractiveTool* Tool, EToolShutdownType ShutdownType)
+{
+	bool bIsActiveTool = (Tool != nullptr) && (ActiveLeftTool == Tool || ActiveRightTool == Tool);
+	if (!bIsActiveTool)
+	{
+		return false;
+	}
+
+	bool bHandled = false;
+	if (OnToolShutdownRequest.IsBound())
+	{
+		bHandled = OnToolShutdownRequest.Execute(this, Tool, ShutdownType);
+	}
+	if (!bHandled)
+	{
+		EToolSide WhichSide = (ActiveLeftTool == Tool) ? EToolSide::Left : EToolSide::Right;
+		DeactivateTool(WhichSide, ShutdownType);
+	}
+	return true;
+}
+
 
 
 
diff --git a/Engine/Source/Runtime/InteractiveToolsFramework/Public/InteractiveToolManager.h b/Engine/Source/Runtime/InteractiveToolsFramework/Public/InteractiveToolManager.h
--- a/Engine/Source/Runtime/InteractiveToolsFramework/Public/InteractiveToolManager.h
+++ b/Engine/Source/Runtime/InteractiveToolsFramework/Public/InteractiveToolManager.h
@@ -238,6 +238,23 @@ public:
 	/** Let active Tools do their screen space drawing. Called by UInteractiveToolsContext. */
 	virtual void DrawHUD(FCanvas* Canvas, IToolsContextRenderAPI* RenderAPI);
 
+	/**
+	 * A Tool (or other code) can call this function to request that the Tool be deactivated.
+	 * The Tool must be active. By default, ::DeactivateTool() wil be called. However if some
+	 * external code has bound to the OnToolShutdownRequest delegate, the request will be
+	 * forwarded to that function first (::DeactivateTool() will still be called if it returns false)
+	 */
+	virtual bool PostActiveToolShutdownRequest(UInteractiveTool* Tool, EToolShutdownType ShutdownType);
+
+	DECLARE_DELEGATE_RetVal_ThreeParams(bool, FOnToolShutdownRequest, UInteractiveToolManager*, UInteractiveTool*, EToolShutdownType);
+	/**
+	 * If bound, OnToolShutdownRequest is called by PostActiveToolShutdownRequest to optionally handle
+	 * requests to shut down a Tool (which may be sent by the Tool itself). Return true to indicate
+	 * that the request will be handled, otherwise ::DeactivateTool() will be called
+	 */
+	FOnToolShutdownRequest OnToolShutdownRequest;
+	 
+
 
 	//
 	// access to APIs, etc
