Gizmo element library: add enable flags based on view projection type and interaction state.

diff --git a/Engine/Source/Runtime/InteractiveToolsFramework/Private/BaseGizmos/GizmoElementArrow.cpp b/Engine/Source/Runtime/InteractiveToolsFramework/Private/BaseGizmos/GizmoElementArrow.cpp
--- a/Engine/Source/Runtime/InteractiveToolsFramework/Private/BaseGizmos/GizmoElementArrow.cpp
+++ b/Engine/Source/Runtime/InteractiveToolsFramework/Private/BaseGizmos/GizmoElementArrow.cpp
@@ -19,11 +19,6 @@ UGizmoElementArrow::UGizmoElementArrow()
 
 void UGizmoElementArrow::Render(IToolsContextRenderAPI* RenderAPI, const FRenderTraversalState& RenderState)
 {
-	if (!IsVisible())
-	{
-		return;
-	}
-
 	check(RenderAPI);
 
 	if (bUpdateArrowBody)
@@ -59,11 +54,6 @@ void UGizmoElementArrow::Render(IToolsContextRenderAPI* RenderAPI, const FRender
 
 FInputRayHit UGizmoElementArrow::LineTrace(const UGizmoViewContext* ViewContext, const FLineTraceTraversalState& LineTraceState, const FVector& RayOrigin, const FVector& RayDirection)
 {
-	if (!IsHittable())
-	{
-		return FInputRayHit();
-	}
-
 	FLineTraceTraversalState CurrentLineTraceState(LineTraceState);
 	bool bHittableViewDependent = UpdateLineTraceState(ViewContext, Base, CurrentLineTraceState);
 
diff --git a/Engine/Source/Runtime/InteractiveToolsFramework/Private/BaseGizmos/GizmoElementBase.cpp b/Engine/Source/Runtime/InteractiveToolsFramework/Private/BaseGizmos/GizmoElementBase.cpp
--- a/Engine/Source/Runtime/InteractiveToolsFramework/Private/BaseGizmos/GizmoElementBase.cpp
+++ b/Engine/Source/Runtime/InteractiveToolsFramework/Private/BaseGizmos/GizmoElementBase.cpp
@@ -37,24 +37,6 @@ bool UGizmoElementBase::GetViewDependentVisibility(const UGizmoViewContext* View
 	return GetViewDependentVisibility(ViewLocation, ViewDirection, bIsPerspectiveView, InLocalToWorldTransform, InLocalCenter);
 }
 
-
-bool UGizmoElementBase::GetViewDependentHittable(const FSceneView* View, const FTransform& InLocalToWorldTransform, const FVector& InLocalCenter) const
-{
-	bool bIsPerspectiveView;
-	FVector ViewLocation, ViewDirection, ViewUp;
-	GizmoElementBaseLocals::GetViewInfo<FSceneView>(View, ViewLocation, ViewDirection, ViewUp, bIsPerspectiveView);
-	return GetViewDependentHittable(ViewLocation, ViewDirection, bIsPerspectiveView, InLocalToWorldTransform, InLocalCenter);
-}
-
-bool UGizmoElementBase::GetViewDependentHittable(const UGizmoViewContext* View, const FTransform& InLocalToWorldTransform, const FVector& InLocalCenter) const
-{
-	bool bIsPerspectiveView;
-	FVector ViewLocation, ViewDirection, ViewUp;
-	GizmoElementBaseLocals::GetViewInfo<UGizmoViewContext>(View, ViewLocation, ViewDirection, ViewUp, bIsPerspectiveView);
-	return GetViewDependentHittable(ViewLocation, ViewDirection, bIsPerspectiveView, InLocalToWorldTransform, InLocalCenter);
-}
-
-
 bool UGizmoElementBase::GetViewAlignRot(const FSceneView* View, const FTransform& InLocalToWorldTransform, const FVector& InLocalCenter, FQuat& OutAlignRot) const
 {
 	bool bIsPerspectiveView;
@@ -104,11 +86,6 @@ bool UGizmoElementBase::GetViewDependentVisibility(const FVector& InViewLocation
 
 }
 
-bool UGizmoElementBase::GetViewDependentHittable(const FVector& InViewLocation, const FVector& InViewDirection, bool bInPerspectiveView, const FTransform& InLocalToWorldTransform, const FVector& InLocalCenter) const
-{
-	return (IsHittable() && (!IsVisible() || GetViewDependentVisibility(InViewLocation, InViewDirection, bInPerspectiveView, InLocalToWorldTransform, InLocalCenter)));
-}
-
 bool UGizmoElementBase::GetViewAlignRot(const FVector& InViewLocation, const FVector& InViewDirection, const FVector& InViewUp, bool bInPerspectiveView, const FTransform& InLocalToWorldTransform, const FVector& InLocalCenter, FQuat& OutAlignRot) const
 {
 	FVector Scale = InLocalToWorldTransform.GetScale3D();
@@ -208,6 +185,62 @@ FQuat UGizmoElementBase::GetAlignRotBetweenCoordSpaces(FVector SourceForward, FV
 	return Result;
 }
 
+bool UGizmoElementBase::GetEnabledForCurrentState(bool bIsPerspectiveProjection) const
+{
+	return (GetEnabledForInteractionState(ElementInteractionState) && GetEnabledForViewProjection(bIsPerspectiveProjection));
+}
+
+bool UGizmoElementBase::GetEnabledForInteractionState(const EGizmoElementInteractionState InElementInteractionState) const
+{
+	if (bEnabled)
+	{
+		switch (InElementInteractionState)
+		{
+		case EGizmoElementInteractionState::None:
+			return bEnabledForDefaultState;		
+
+		case EGizmoElementInteractionState::Hovering:
+			return bEnabledForHoveringState;
+
+		case EGizmoElementInteractionState::Interacting:	
+			return bEnabledForInteractingState;
+		}
+	}
+
+	return false;
+}
+
+bool UGizmoElementBase::GetEnabledForViewProjection(bool bIsPerspectiveProjection) const
+{
+	if (bEnabled)
+	{
+		if (bIsPerspectiveProjection)
+		{
+			return bEnabledForPerspectiveProjection;
+		}
+		else
+		{
+			return bEnabledForOrthographicProjection;
+		}
+	}
+	
+	return false;
+}
+
+bool UGizmoElementBase::IsVisible(const FSceneView* View, const FTransform& InLocalToWorldTransform, const FVector& InLocalCenter) const
+{
+	return (GetVisibleState() &&
+		GetEnabledForCurrentState(View->IsPerspectiveProjection()) &&
+		GetViewDependentVisibility(View, InLocalToWorldTransform, InLocalCenter));
+}
+
+bool UGizmoElementBase::IsHittable(const UGizmoViewContext* ViewContext, const FTransform& InLocalToWorldTransform, const FVector& InLocalCenter) const
+{
+	return (GetHittableState() &&
+		GetEnabledForCurrentState(ViewContext->IsPerspectiveProjection()) &&
+		(!GetVisibleState() || GetViewDependentVisibility(ViewContext, InLocalToWorldTransform, InLocalCenter)));
+}
+
 bool UGizmoElementBase::UpdateRenderState(IToolsContextRenderAPI* RenderAPI, const FVector& InLocalCenter, FRenderTraversalState& InOutRenderState)
 {
 	FQuat AlignRot;
@@ -231,15 +264,14 @@ bool UGizmoElementBase::UpdateRenderState(IToolsContextRenderAPI* RenderAPI, con
 
 	InOutRenderState.MeshRenderState.Update(MeshRenderAttributes);
 
-	bool bVisibleViewDependent = GetViewDependentVisibility(View, InOutRenderState.LocalToWorldTransform, InLocalCenter);
-
-	if (bVisibleViewDependent)
+	if (IsVisible(View, InOutRenderState.LocalToWorldTransform, InLocalCenter))
 	{
 		bOutHasAlignRot = GetViewAlignRot(View, InOutRenderState.LocalToWorldTransform, InLocalCenter, OutAlignRot);
 		InOutRenderState.LocalToWorldTransform = FTransform(OutAlignRot, InLocalCenter) * InOutRenderState.LocalToWorldTransform;
+		return true;
 	}
 
-	return bVisibleViewDependent;
+	return false;
 }
 
 bool UGizmoElementBase::UpdateLineTraceState(const UGizmoViewContext* ViewContext, const FVector& InLocalCenter, FLineTraceTraversalState& InOutRenderState)
@@ -256,27 +288,24 @@ bool UGizmoElementBase::UpdateLineTraceState(const UGizmoViewContext* ViewContex
 	OutAlignRot = FQuat::Identity;
 	bOutHasAlignRot = false;
 
-	bool bHittableViewDependent = GetViewDependentHittable(ViewContext, InOutRenderState.LocalToWorldTransform, InLocalCenter);
-
-	if (bHittableViewDependent)
+	if (IsHittable(ViewContext, InOutRenderState.LocalToWorldTransform, InLocalCenter))
 	{
 		bOutHasAlignRot = GetViewAlignRot(ViewContext, InOutRenderState.LocalToWorldTransform, InLocalCenter, OutAlignRot);
 		InOutRenderState.LocalToWorldTransform = FTransform(OutAlignRot, InLocalCenter) * InOutRenderState.LocalToWorldTransform;
+		return true;
 	}
 
-	return bHittableViewDependent;
+	return false;
 }
 
-bool UGizmoElementBase::IsVisible() const
+bool UGizmoElementBase::GetVisibleState() const
 {
-	bool bVisible = static_cast<uint8>(ElementState) & static_cast<uint8>(EGizmoElementState::Visible);
-	return (bEnabled && bVisible);
+	return (static_cast<uint8>(ElementState) & static_cast<uint8>(EGizmoElementState::Visible));
 }
 
-bool UGizmoElementBase::IsHittable() const
+bool UGizmoElementBase::GetHittableState() const
 {
-	bool bHittable = static_cast<uint8>(ElementState) & static_cast<uint8>(EGizmoElementState::Hittable);
-	return (bEnabled && bHittable);
+	return (static_cast<uint8>(ElementState) & static_cast<uint8>(EGizmoElementState::Hittable));
 }
 
 void UGizmoElementBase::SetEnabled(bool InEnabled)
@@ -288,6 +317,57 @@ bool UGizmoElementBase::GetEnabled() const
 {
 	return bEnabled;
 }
+
+void UGizmoElementBase::SetEnabledForPerspectiveProjection(bool bInEnabledForPerspectiveProjection)
+{
+	bEnabledForPerspectiveProjection = bInEnabledForPerspectiveProjection;
+}
+
+bool UGizmoElementBase::GetEnabledForPerspectiveProjection()
+{
+	return bEnabledForPerspectiveProjection;
+}
+
+void UGizmoElementBase::SetEnabledInOrthographicProjection(bool bInEnabledForOrthographicProjection)
+{
+	bEnabledForOrthographicProjection = bInEnabledForOrthographicProjection;
+}
+
+bool UGizmoElementBase::GetEnabledInOrthographicProjection()
+{
+	return bEnabledForOrthographicProjection;
+}
+
+void UGizmoElementBase::SetEnabledForDefaultState(bool bInEnabledForDefaultState)
+{
+	bEnabledForDefaultState = bInEnabledForDefaultState;
+}
+
+bool UGizmoElementBase::GetEnabledForDefaultState()
+{
+	return bEnabledForDefaultState;
+}
+
+void UGizmoElementBase::SetEnabledForHoveringState(bool bInEnabledForHoveringState)
+{
+	bEnabledForHoveringState = bInEnabledForHoveringState;
+}
+
+bool UGizmoElementBase::GetEnabledForHoveringState()
+{
+	return bEnabledForHoveringState;
+}
+
+void UGizmoElementBase::SetEnabledForInteractingState(bool bInEnabledForInteractingState)
+{
+	bEnabledForInteractingState = bInEnabledForInteractingState;
+}
+
+bool UGizmoElementBase::GetEnabledForInteractingState()
+{
+	return bEnabledForInteractingState;
+}
+
 void UGizmoElementBase::SetPartIdentifier(uint32 InPartIdentifier)
 {
 	PartIdentifier = InPartIdentifier;
diff --git a/Engine/Source/Runtime/InteractiveToolsFramework/Private/BaseGizmos/GizmoElementBox.cpp b/Engine/Source/Runtime/InteractiveToolsFramework/Private/BaseGizmos/GizmoElementBox.cpp
--- a/Engine/Source/Runtime/InteractiveToolsFramework/Private/BaseGizmos/GizmoElementBox.cpp
+++ b/Engine/Source/Runtime/InteractiveToolsFramework/Private/BaseGizmos/GizmoElementBox.cpp
@@ -9,13 +9,6 @@
 
 void UGizmoElementBox::Render(IToolsContextRenderAPI* RenderAPI, const FRenderTraversalState& RenderState)
 {
-	if (!IsVisible())
-	{
-		return;
-	}
-
-	check(RenderAPI);
-
 	FRenderTraversalState CurrentRenderState(RenderState);
 	bool bVisibleViewDependent = UpdateRenderState(RenderAPI, Center, CurrentRenderState);
 
@@ -34,11 +27,6 @@ void UGizmoElementBox::Render(IToolsContextRenderAPI* RenderAPI, const FRenderTr
 
 FInputRayHit UGizmoElementBox::LineTrace(const UGizmoViewContext* ViewContext, const FLineTraceTraversalState& LineTraceState, const FVector& RayOrigin, const FVector& RayDirection)
 {
-	if (!IsHittable())
-	{
-		return FInputRayHit();
-	}
-
 	FLineTraceTraversalState CurrentLineTraceState(LineTraceState);
 	bool bHittableViewDependent = UpdateLineTraceState(ViewContext, Center, CurrentLineTraceState);
 
diff --git a/Engine/Source/Runtime/InteractiveToolsFramework/Private/BaseGizmos/GizmoElementCircle.cpp b/Engine/Source/Runtime/InteractiveToolsFramework/Private/BaseGizmos/GizmoElementCircle.cpp
--- a/Engine/Source/Runtime/InteractiveToolsFramework/Private/BaseGizmos/GizmoElementCircle.cpp
+++ b/Engine/Source/Runtime/InteractiveToolsFramework/Private/BaseGizmos/GizmoElementCircle.cpp
@@ -10,13 +10,6 @@
 
 void UGizmoElementCircle::Render(IToolsContextRenderAPI* RenderAPI, const FRenderTraversalState& RenderState)
 {
-	if (!IsVisible())
-	{
-		return;
-	}
-
-	check(RenderAPI);
-
 	FRenderTraversalState CurrentRenderState(RenderState);
 	bool bVisibleViewDependent = UpdateRenderState(RenderAPI, Center, CurrentRenderState);
 
@@ -50,11 +43,6 @@ void UGizmoElementCircle::Render(IToolsContextRenderAPI* RenderAPI, const FRende
 
 FInputRayHit UGizmoElementCircle::LineTrace(const UGizmoViewContext* ViewContext, const FLineTraceTraversalState& LineTraceState, const FVector& RayOrigin, const FVector& RayDirection)
 {
-	if (!IsHittable())
-	{
-		return FInputRayHit();
-	}
-
 	FLineTraceTraversalState CurrentLineTraceState(LineTraceState);
 	bool bHittableViewDependent = UpdateLineTraceState(ViewContext, Center, CurrentLineTraceState);
 
diff --git a/Engine/Source/Runtime/InteractiveToolsFramework/Private/BaseGizmos/GizmoElementCone.cpp b/Engine/Source/Runtime/InteractiveToolsFramework/Private/BaseGizmos/GizmoElementCone.cpp
--- a/Engine/Source/Runtime/InteractiveToolsFramework/Private/BaseGizmos/GizmoElementCone.cpp
+++ b/Engine/Source/Runtime/InteractiveToolsFramework/Private/BaseGizmos/GizmoElementCone.cpp
@@ -9,11 +9,6 @@
 
 void UGizmoElementCone::Render(IToolsContextRenderAPI* RenderAPI, const FRenderTraversalState& RenderState)
 {
-	if (!IsVisible())
-	{
-		return;
-	}
-
 	FRenderTraversalState CurrentRenderState(RenderState);
 	bool bVisibleViewDependent = UpdateRenderState(RenderAPI, Origin, CurrentRenderState);
 
@@ -35,11 +30,6 @@ void UGizmoElementCone::Render(IToolsContextRenderAPI* RenderAPI, const FRenderT
 
 FInputRayHit UGizmoElementCone::LineTrace(const UGizmoViewContext* ViewContext, const FLineTraceTraversalState& LineTraceState, const FVector& RayOrigin, const FVector& RayDirection)
 {
-	if (!IsHittable())
-	{
-		return FInputRayHit();
-	}
-
 	FLineTraceTraversalState CurrentLineTraceState(LineTraceState);
 	bool bHittableViewDependent = UpdateLineTraceState(ViewContext, Origin, CurrentLineTraceState);
 
diff --git a/Engine/Source/Runtime/InteractiveToolsFramework/Private/BaseGizmos/GizmoElementCylinder.cpp b/Engine/Source/Runtime/InteractiveToolsFramework/Private/BaseGizmos/GizmoElementCylinder.cpp
--- a/Engine/Source/Runtime/InteractiveToolsFramework/Private/BaseGizmos/GizmoElementCylinder.cpp
+++ b/Engine/Source/Runtime/InteractiveToolsFramework/Private/BaseGizmos/GizmoElementCylinder.cpp
@@ -9,13 +9,6 @@
 
 void UGizmoElementCylinder::Render(IToolsContextRenderAPI* RenderAPI, const FRenderTraversalState& RenderState)
 {
-	if (!IsVisible())
-	{
-		return;
-	}
-
-	check(RenderAPI);
-
 	FRenderTraversalState CurrentRenderState(RenderState);
 	bool bVisibleViewDependent = UpdateRenderState(RenderAPI, Base, CurrentRenderState);
 
@@ -37,11 +30,6 @@ void UGizmoElementCylinder::Render(IToolsContextRenderAPI* RenderAPI, const FRen
 
 FInputRayHit UGizmoElementCylinder::LineTrace(const UGizmoViewContext* ViewContext, const FLineTraceTraversalState& LineTraceState, const FVector& RayOrigin, const FVector& RayDirection)
 {
-	if (!IsHittable())
-	{
-		return FInputRayHit();
-	}
-
 	FLineTraceTraversalState CurrentLineTraceState(LineTraceState);
 	bool bHittableViewDependent = UpdateLineTraceState(ViewContext, Base, CurrentLineTraceState);
 
diff --git a/Engine/Source/Runtime/InteractiveToolsFramework/Private/BaseGizmos/GizmoElementGroup.cpp b/Engine/Source/Runtime/InteractiveToolsFramework/Private/BaseGizmos/GizmoElementGroup.cpp
--- a/Engine/Source/Runtime/InteractiveToolsFramework/Private/BaseGizmos/GizmoElementGroup.cpp
+++ b/Engine/Source/Runtime/InteractiveToolsFramework/Private/BaseGizmos/GizmoElementGroup.cpp
@@ -4,7 +4,7 @@
 #include "BaseGizmos/GizmoViewContext.h"
 
 
-void UGizmoElementGroup::ApplyConstantScale(float PixelToWorldScale, FTransform& InOutLocalToWorldTransform)
+void UGizmoElementGroup::ApplyUniformConstantScaleToTransform(float PixelToWorldScale, FTransform& InOutLocalToWorldTransform) const
 {
 	float Scale = InOutLocalToWorldTransform.GetScale3D().X;
 	if (bConstantScale)
@@ -16,19 +16,12 @@ void UGizmoElementGroup::ApplyConstantScale(float PixelToWorldScale, FTransform&
 
 void UGizmoElementGroup::Render(IToolsContextRenderAPI* RenderAPI, const FRenderTraversalState& RenderState)
 {
-	if (!IsVisible())
-	{
-		return;
-	}
-
-	check(RenderAPI);
-
 	FRenderTraversalState CurrentRenderState(RenderState);
 	bool bVisibleViewDependent = UpdateRenderState(RenderAPI, FVector::ZeroVector, CurrentRenderState);
 
 	if (bVisibleViewDependent)
 	{
-		ApplyConstantScale(CurrentRenderState.PixelToWorldScale, CurrentRenderState.LocalToWorldTransform);
+		ApplyUniformConstantScaleToTransform(CurrentRenderState.PixelToWorldScale, CurrentRenderState.LocalToWorldTransform);
 
 		// Continue render even if not visible so all transforms will be cached 
 		// for subsequent line tracing.
@@ -51,7 +44,7 @@ FInputRayHit UGizmoElementGroup::LineTrace(const UGizmoViewContext* ViewContext,
 
 	if (bHittableViewDependent)
 	{
-		ApplyConstantScale(CurrentLineTraceState.PixelToWorldScale, CurrentLineTraceState.LocalToWorldTransform);
+		ApplyUniformConstantScaleToTransform(CurrentLineTraceState.PixelToWorldScale, CurrentLineTraceState.LocalToWorldTransform);
 
 		for (UGizmoElementBase* Element : Elements)
 		{
diff --git a/Engine/Source/Runtime/InteractiveToolsFramework/Private/BaseGizmos/GizmoElementRectangle.cpp b/Engine/Source/Runtime/InteractiveToolsFramework/Private/BaseGizmos/GizmoElementRectangle.cpp
--- a/Engine/Source/Runtime/InteractiveToolsFramework/Private/BaseGizmos/GizmoElementRectangle.cpp
+++ b/Engine/Source/Runtime/InteractiveToolsFramework/Private/BaseGizmos/GizmoElementRectangle.cpp
@@ -10,13 +10,6 @@
 
 void UGizmoElementRectangle::Render(IToolsContextRenderAPI* RenderAPI, const FRenderTraversalState& RenderState)
 {
-	if (!IsVisible())
-	{
-		return;
-	}
-
-	check(RenderAPI);
-
 	FRenderTraversalState CurrentRenderState(RenderState);
 	bool bVisibleViewDependent = UpdateRenderState(RenderAPI, Center, CurrentRenderState);
 
@@ -48,11 +41,6 @@ void UGizmoElementRectangle::Render(IToolsContextRenderAPI* RenderAPI, const FRe
 
 FInputRayHit UGizmoElementRectangle::LineTrace(const UGizmoViewContext* ViewContext, const FLineTraceTraversalState& LineTraceState, const FVector& RayOrigin, const FVector& RayDirection)
 {
-	if (!IsHittable())
-	{
-		return FInputRayHit();
-	}
-
 	FLineTraceTraversalState CurrentLineTraceState(LineTraceState);
 	bool bHittableViewDependent = UpdateLineTraceState(ViewContext, Center, CurrentLineTraceState);
 
diff --git a/Engine/Source/Runtime/InteractiveToolsFramework/Private/BaseGizmos/GizmoElementTorus.cpp b/Engine/Source/Runtime/InteractiveToolsFramework/Private/BaseGizmos/GizmoElementTorus.cpp
--- a/Engine/Source/Runtime/InteractiveToolsFramework/Private/BaseGizmos/GizmoElementTorus.cpp
+++ b/Engine/Source/Runtime/InteractiveToolsFramework/Private/BaseGizmos/GizmoElementTorus.cpp
@@ -9,14 +9,6 @@
 
 void UGizmoElementTorus::Render(IToolsContextRenderAPI* RenderAPI, const FRenderTraversalState& RenderState)
 {
-	if (!IsVisible())
-	{
-		return;
-	}
-
-	check(RenderAPI);
-
-
 	FRenderTraversalState CurrentRenderState(RenderState);
 	bool bVisibleViewDependent = UpdateRenderState(RenderAPI, Center, CurrentRenderState);
 
diff --git a/Engine/Source/Runtime/InteractiveToolsFramework/Public/BaseGizmos/GizmoElementBase.h b/Engine/Source/Runtime/InteractiveToolsFramework/Public/BaseGizmos/GizmoElementBase.h
--- a/Engine/Source/Runtime/InteractiveToolsFramework/Public/BaseGizmos/GizmoElementBase.h
+++ b/Engine/Source/Runtime/InteractiveToolsFramework/Public/BaseGizmos/GizmoElementBase.h
@@ -7,6 +7,7 @@
 #include "BaseGizmos/GizmoElementShared.h"
 #include "BaseGizmos/GizmoInterfaces.h"
 #include "BaseGizmos/GizmoRenderingUtil.h"
+#include "BaseGizmos/GizmoViewContext.h"
 #include "InputState.h"
 #include "SceneView.h"
 #include "ToolContextInterfaces.h"
@@ -14,7 +15,6 @@
 #include "GizmoElementBase.generated.h"
 
 class UMaterialInterface;
-class UGizmoViewContext;
 
 /**
  * Base class for 2d and 3d primitive objects intended to be used as part of 3D Gizmos.
@@ -98,11 +98,16 @@ public:
 		// Pixel to world scale
 		double PixelToWorldScale = 1.0;
 
+		// View context is perspective projecion
+		bool bIsPerspectiveProjection = true;
+
 		// Initialize state 
 		void Initialize(const UGizmoViewContext* InGizmoViewContext, FTransform InTransform)
 		{
+			check(InGizmoViewContext);
 			LocalToWorldTransform = InTransform;
 			PixelToWorldScale = GizmoRenderingUtil::CalculateLocalPixelToWorldScale(InGizmoViewContext, InTransform.GetLocation());
+			bIsPerspectiveProjection = InGizmoViewContext->IsPerspectiveProjection();
 		}
 	};
 
@@ -117,16 +122,37 @@ public:
 	// Calcute box sphere bounds for use when hit testing.
 	virtual FBoxSphereBounds CalcBounds(const FTransform& LocalToWorld) const PURE_VIRTUAL(UGizmoElementBase::CalcBounds, return FBoxSphereBounds(););
 
-	// Whether this object is visible.
-	virtual bool IsVisible() const;
+	// Return whether the visible bit is enabled in element state.
+	virtual bool GetVisibleState() const;
 
-	// Whether this object is hittable.
-	virtual bool IsHittable() const;
+	// Return whether the hittable bit is enabled in element state.
+	virtual bool GetHittableState() const;
 
-	// Render and LineTrace should only occur when bEnabled is true.
-	virtual void SetEnabled(bool InEnabled);
+	// Element enabled flag. Render and LineTrace only occur when bEnabled is true.
+	// This flag is useful for turning on and off an element globally while retaining its specific settings.
+	virtual void SetEnabled(bool bInEnabled);
 	virtual bool GetEnabled() const;
 
+	// Whether element is enabled for perspective projections.
+	virtual void SetEnabledForPerspectiveProjection(bool bInEnabledForPerspectiveProjection);
+	virtual bool GetEnabledForPerspectiveProjection();
+
+	// Whether element is enabled for orthographic projections.
+	virtual void SetEnabledInOrthographicProjection(bool bInEnabledForOrthographicProjection);
+	virtual bool GetEnabledInOrthographicProjection();
+
+	// Whether element is enabled when element state is default.
+	virtual void SetEnabledForDefaultState(bool bInEnabledForDefaultState);
+	virtual bool GetEnabledForDefaultState();
+
+	// Whether element is enabled when element state is hovering.
+	virtual void SetEnabledForHoveringState(bool bInEnabledForHoveringState);
+	virtual bool GetEnabledForHoveringState();
+
+	// Whether element is enabled when element state is interacting.
+	virtual void SetEnabledForInteractingState(bool bInEnabledForInteractingState);
+	virtual bool GetEnabledForInteractingState();
+
 	// For an element hierarchy representing multiple parts of a single gizmo, the part identifier establishes 
 	// a correspondence between a gizmo part and the elements that represent that part. The recognized
 	// part identifier values should be defined in the gizmo. Gizmo part identifiers must be greater than or 
@@ -267,6 +293,26 @@ protected:
 	UPROPERTY()
 	bool bEnabled = true;
 
+	// Whether element is enabled for perspective projection
+	UPROPERTY()
+	bool bEnabledForPerspectiveProjection = true;
+
+	// Whether element is enabled for orthographic projection
+	UPROPERTY()
+	bool bEnabledForOrthographicProjection = true;
+
+	// Whether element is enabled when element state is default
+	UPROPERTY()
+	bool bEnabledForDefaultState = true;
+
+	// Whether element is enabled when element state is hovering
+	UPROPERTY()
+	bool bEnabledForHoveringState = true;
+
+	// Whether element is enabled when element state is interacting
+	UPROPERTY()
+	bool bEnabledForInteractingState = true;
+
 	// Part identifier
 	UPROPERTY()
 	uint32 PartIdentifier = DefaultPartIdentifier;
@@ -279,7 +325,7 @@ protected:
 	UPROPERTY()
 	EGizmoElementState ElementState = EGizmoElementState::VisibleAndHittable;
 
-	// Element interaction state - None, Hovering or Interacting
+	// Current element interaction state - None, Hovering or Interacting
 	UPROPERTY()
 	EGizmoElementInteractionState ElementInteractionState = EGizmoElementInteractionState::None;
 
@@ -341,23 +387,29 @@ protected:
 
 protected:
 
-	// Returns whether object is visible in input FSceneView based on view-dependent visibility settings 
+	// Return whether element is currently visible.
+	virtual bool IsVisible(const FSceneView* View, const FTransform& InLocalToWorldTransform, const FVector& InLocalCenter) const;
+
+	// Return whether element is currently hittable.
+	virtual bool IsHittable(const UGizmoViewContext* ViewContext, const FTransform& InLocalToWorldTransform, const FVector& InLocalCenter) const;
+
+	// Returns whether object is visible in input FSceneView based on view-dependent visibility settings.
 	virtual bool GetViewDependentVisibility(const FSceneView* View, const FTransform& InLocalToWorldTransform, const FVector& InLocalCenter) const;
 
 	// Returns whether object is visible in input gizmo view context based on view-dependent visibility settings
 	virtual bool GetViewDependentVisibility(const UGizmoViewContext* View, const FTransform& InLocalToWorldTransform, const FVector& InLocalCenter) const;
 
-	// Returns whether object is visible based input view parameters and view-dependent visibility settings
+	// Returns whether object is visible based on view-dependent visibility settings.
 	virtual bool GetViewDependentVisibility(const FVector& InViewLocation, const FVector& InViewDirection, bool bInPerspectiveView, const FTransform& InLocalToWorldTransform, const FVector& InLocalCenter) const;
 
-	// Returns whether object is hittable in input FSceneView based on element state and view-dependent settings. 
-	virtual bool GetViewDependentHittable(const FSceneView* View, const FTransform& InLocalToWorldTransform, const FVector& InLocalCenter) const;
+	// Returns whether element is enabled based on element's current interaction state and input view projection type.
+	virtual bool GetEnabledForCurrentState(bool bIsPerspectiveProjection) const;
 
-	// Returns whether object is hittable in input gizmo view context based on element state and view-dependent settings. 
-	virtual bool GetViewDependentHittable(const UGizmoViewContext* View, const FTransform& InLocalToWorldTransform, const FVector& InLocalCenter) const;
+	// Returns whether element is enabled for given interaction state.
+	virtual bool GetEnabledForInteractionState(EGizmoElementInteractionState InInteractionState) const;
 
-	// Returns whether object is hittable based on input view parameters, element state and view-dependent settings. 
-	virtual bool GetViewDependentHittable(const FVector& InViewLocation, const FVector& InViewDirection, bool bInPerspectiveView, const FTransform& InLocalToWorldTransform, const FVector& InLocalCenter) const;
+	// Returns whether element is enabled for given view projection type.
+	virtual bool GetEnabledForViewProjection(bool bIsPerspectiveProjection) const;
 
 	// Return whether this element has a view alignment rotation based on input FSceneView and view-dependent alignment settings.
 	// @param OutAlignRot the rotation to align this element in local space, should be prepended to the local-to-world transform.
@@ -378,7 +430,7 @@ protected:
 	// @return view dependent visibility, true if this element is visible in the current view. 
 	virtual bool UpdateRenderState(IToolsContextRenderAPI* RenderAPI, const FVector& InLocalCenter, FRenderTraversalState& InOutRenderState);
 
-	// Update render state during render traversal, determines the current render state for this element, same parameters as
+	// Update render state during render traversal, determines the current render state for this element
 	// Same parameters as UpdateRenderState above plus two output parameters:
 	// @param bOutHasAlignRot - whether alignment rotation was applied to output state
 	// @param OutAlignRot - alignment rotation applied to output state, if applicable.
diff --git a/Engine/Source/Runtime/InteractiveToolsFramework/Public/BaseGizmos/GizmoElementGroup.h b/Engine/Source/Runtime/InteractiveToolsFramework/Public/BaseGizmos/GizmoElementGroup.h
--- a/Engine/Source/Runtime/InteractiveToolsFramework/Public/BaseGizmos/GizmoElementGroup.h
+++ b/Engine/Source/Runtime/InteractiveToolsFramework/Public/BaseGizmos/GizmoElementGroup.h
@@ -59,6 +59,6 @@ protected:
 	UPROPERTY()
 	TArray<TObjectPtr<UGizmoElementBase>> Elements;
 
-	// Applies constant scale to the current local to world transform, if constant scale is true
-	virtual void ApplyConstantScale(float PixelToWorldScale, FTransform& InOutLocalToWorldTransform);
+	// Updates input transform's scale component to have uniform scale and applies constant scale if bConstantScale is true
+	virtual void ApplyUniformConstantScaleToTransform(float PixelToWorldScale, FTransform& InOutLocalToWorldTransform) const;
 };
