Add approximate LineTrace implementation for GizmoElementTorus. Remove unused method.

diff --git a/Engine/Source/Runtime/InteractiveToolsFramework/Private/BaseGizmos/GizmoElementTorus.cpp b/Engine/Source/Runtime/InteractiveToolsFramework/Private/BaseGizmos/GizmoElementTorus.cpp
--- a/Engine/Source/Runtime/InteractiveToolsFramework/Private/BaseGizmos/GizmoElementTorus.cpp
+++ b/Engine/Source/Runtime/InteractiveToolsFramework/Private/BaseGizmos/GizmoElementTorus.cpp
@@ -21,22 +21,143 @@ void UGizmoElementTorus::Render(IToolsContextRenderAPI* RenderAPI, const FRender
 			FVector TorusSideAxis = Normal ^ BeginAxis;
 			TorusSideAxis.Normalize();
 
-			DrawTorus(PDI, CurrentRenderState.LocalToWorldTransform.ToMatrixWithScale(), BeginAxis, TorusSideAxis, OuterRadius, InnerRadius, OuterSegments, InnerSlices,  UseMaterial->GetRenderProxy(), SDPG_Foreground, bPartial, Angle, bEndCaps);
+			DrawTorus(PDI, CurrentRenderState.LocalToWorldTransform.ToMatrixWithScale(), 
+				BeginAxis, TorusSideAxis, OuterRadius, InnerRadius, OuterSegments, InnerSlices,  
+				UseMaterial->GetRenderProxy(), SDPG_Foreground, bPartial, Angle, bEndCaps);
 		}
 	}
 }
 
 
-
+//
+// This method approximates ray-torus intersection by intersecting the ray with the plane in which the torus 
+// lies, then determining a hit point closest to the linear circle defined by torus center and torus outer radius.
+// 
+// If the torus lies at a glancing angle, the ray-torus intersection is performed against cylinders approximating
+// the shape of the torus.
+//
 FInputRayHit UGizmoElementTorus::LineTrace(const UGizmoViewContext* ViewContext, const FLineTraceTraversalState& LineTraceState, const FVector& RayOrigin, const FVector& RayDirection)
 {
-	// Line trace is not supported for torus.
+	FLineTraceTraversalState CurrentLineTraceState(LineTraceState);
+	bool bHittableViewDependent = UpdateLineTraceState(ViewContext, Center, CurrentLineTraceState);
+
+	if (bHittableViewDependent)
+	{
+		const double PixelHitThresholdAdjust = CurrentLineTraceState.PixelToWorldScale * PixelHitDistanceThreshold;
+		const double WorldOuterRadius = OuterRadius * CurrentLineTraceState.LocalToWorldTransform.GetScale3D().X;
+		const double WorldInnerRadius = InnerRadius * CurrentLineTraceState.LocalToWorldTransform.GetScale3D().X;
+		const FVector WorldCenter = CurrentLineTraceState.LocalToWorldTransform.TransformPosition(FVector::ZeroVector);
+		const FVector WorldNormal = CurrentLineTraceState.LocalToWorldTransform.TransformVectorNoScale(Normal);
+		const FVector WorldBeginAxis = CurrentLineTraceState.LocalToWorldTransform.GetRotation().RotateVector(BeginAxis);
+		double HitDepth = -1.0;
+
+
+		// Determine if the ray direction is at a glancing angle by using a minimum cos angle based on the
+		// angle between the vector from torus center to ring center and the vector from torus center to ring edge.
+		double MinCosAngle = OuterRadius / FMath::Sqrt(OuterRadius * OuterRadius + InnerRadius * InnerRadius);
+		bool bAtGlancingAngle = FMath::Abs(FVector::DotProduct(WorldNormal, RayDirection)) <= MinCosAngle;
+
+		if (bAtGlancingAngle)
+		{
+			// If the torus lies at a glancing angle, the ray-torus intersection is performed against cylinders approximating
+			// the shape of the torus.
+			static constexpr int NumFullTorusCylinders = 16;
+			static constexpr double AngleDelta = UE_DOUBLE_TWO_PI / static_cast<double>(NumFullTorusCylinders);
+			const int NumCylinders = bPartial ? FMath::CeilToInt(NumFullTorusCylinders * Angle / UE_DOUBLE_TWO_PI) : NumFullTorusCylinders;
+
+			const FVector ViewDirection = ViewContext->GetViewDirection();
+			FVector VectorA = WorldBeginAxis; 
+			FVector VectorB = VectorA.RotateAngleAxisRad(AngleDelta, WorldNormal);
+
+			const double CylinderRadius = WorldInnerRadius + PixelHitThresholdAdjust;
+			double CylinderHeight = (VectorB - VectorA).Length() * WorldOuterRadius;
+
+			if (FMath::IsNearlyZero(CylinderHeight))
+			{
+				return FInputRayHit();
+			}
+
+			// Line trace against a set of cylinders approximating the shape of the torus
+			for (int i = 0; i < NumCylinders; i++)
+			{
+				if (i > 0)
+				{
+					VectorA = VectorB;
+					VectorB = VectorA.RotateAngleAxisRad(AngleDelta, Normal);
+				}
+
+				if (i == NumCylinders - 1)
+				{
+					CylinderHeight = bPartial ? FMath::Fmod(Angle, AngleDelta) * CylinderHeight : CylinderHeight;
+				}
+
+				const FVector CylinderDirection = (VectorB - VectorA).GetSafeNormal();
+				const FVector CylinderCenter = WorldCenter + (VectorA * WorldOuterRadius) + CylinderDirection * CylinderHeight * 0.5;			
+				bool bIntersects = false;
+				double RayParam = -1.0;
+
+				GizmoMath::RayCylinderIntersection(
+					CylinderCenter,
+					CylinderDirection,
+					CylinderRadius,
+					CylinderHeight,
+					RayOrigin, RayDirection,
+					bIntersects, RayParam);
+
+				// Update with closest hit depth
+				if (bIntersects && (HitDepth < 0 || HitDepth > RayParam))
+				{
+					HitDepth = RayParam;
+				}
+			}
+		}
+		else
+		{
+			// Intersect ray with plane in which torus lies.
+			FPlane Plane(WorldCenter, WorldNormal);
+			HitDepth = FMath::RayPlaneIntersectionParam(RayOrigin, RayDirection, Plane);
+			if (HitDepth < 0.0)
+			{
+				return FInputRayHit();
+			}
+
+			FVector HitPoint = RayOrigin + RayDirection * HitDepth;
+
+			// Find the closest point on the circle to the intersection point
+			FVector NearestCirclePos;
+			GizmoMath::ClosetPointOnCircle(HitPoint, WorldCenter, WorldNormal, WorldOuterRadius, NearestCirclePos);
+
+			// Find the closest point on the ray to the circle and determine if it is within the torus
+			FRay Ray(RayOrigin, RayDirection, true);
+			FVector NearestRayPos = Ray.ClosestPoint(NearestCirclePos);
+
+			const double HitBuffer = PixelHitThresholdAdjust + WorldInnerRadius;
+			double Distance = FVector::Distance(NearestCirclePos, NearestRayPos);
+
+			if (Distance > HitBuffer)
+			{
+				return FInputRayHit();
+			}
+
+			// Adjust hit depth to return the closest hit point's depth
+			HitDepth = (NearestRayPos - RayOrigin).Length();
+		}
+
+		if (HitDepth >= 0.0)
+		{
+			FInputRayHit RayHit(HitDepth);
+			RayHit.SetHitObject(this);
+			RayHit.HitIdentifier = PartIdentifier;
+			return RayHit;
+		}
+	}
+	
 	return FInputRayHit();
 }
 
 FBoxSphereBounds UGizmoElementTorus::CalcBounds(const FTransform& LocalToWorld) const
 {
-	// Box sphere bounds is not supported for torus.
+	// @todo - implement box-sphere bounds calculation
 	return FBoxSphereBounds();
 }
 
@@ -122,16 +243,6 @@ bool UGizmoElementTorus::GetPartial() const
 	return bPartial;
 }
 
-void UGizmoElementTorus::SetScreenAlignPartial(bool InScreenAlignPartial)
-{
-	bScreenAlignPartial = InScreenAlignPartial;
-}
-
-bool UGizmoElementTorus::GetScreenAlignPartial() const
-{
-	return bPartial;
-}
-
 void UGizmoElementTorus::SetAngle(float InAngle)
 {
 	Angle = InAngle;
diff --git a/Engine/Source/Runtime/InteractiveToolsFramework/Public/BaseGizmos/GizmoElementTorus.h b/Engine/Source/Runtime/InteractiveToolsFramework/Public/BaseGizmos/GizmoElementTorus.h
--- a/Engine/Source/Runtime/InteractiveToolsFramework/Public/BaseGizmos/GizmoElementTorus.h
+++ b/Engine/Source/Runtime/InteractiveToolsFramework/Public/BaseGizmos/GizmoElementTorus.h
@@ -14,9 +14,9 @@ class FMaterialRenderProxy;
 /**
  * Simple object intended to be used as part of 3D Gizmos.
  * Draws a torus based on parameters.
- * Note: Ray-torus hit testing is not supported! LineTrace() will never return a hit nor will CalcBounds() calculate box-sphere bounds.
- *       Use a hidden GizmoElementCircle with circle primitive object with DrawLines set to true to approximate 
- *       torus intersection for hit-testing purposes.
+ * 
+ * Note: the LineTrace method does not perform a true ray-torus intersection!
+ * See comment above LineTrace method below for details of how this intersection is approximated.
  */
 UCLASS(Transient)
 class INTERACTIVETOOLSFRAMEWORK_API UGizmoElementTorus : public UGizmoElementBase
@@ -27,6 +27,11 @@ public:
 
 	//~ Begin UGizmoElementBase Interface.
 	virtual void Render(IToolsContextRenderAPI* RenderAPI, const FRenderTraversalState& RenderState) override;
+
+	// LineTrace approximates ray-torus intersection by intersecting the ray with the plane in which the torus 
+	// lies, then determining a hit point closest to the linear circle defined by torus center and torus outer radius.
+	// If the torus lies at a glancing angle, the ray-torus intersection is performed against cylinders approximating
+	// the shape of the torus.
 	virtual FInputRayHit LineTrace(const UGizmoViewContext* ViewContext, const FLineTraceTraversalState& LineTraceState, const FVector& RayOrigin, const FVector& RayDirection) override;
 	virtual FBoxSphereBounds CalcBounds(const FTransform& LocalToWorld) const override;
 	//~ End UGizmoElementBase Interface.
@@ -63,10 +68,6 @@ public:
 	virtual void SetPartial(bool InPartial);
 	virtual bool GetPartial() const;
 
-	// True when the partial torus arc should be screen aligned.
-	virtual void SetScreenAlignPartial(bool InScreenAlignPartial);
-	virtual bool GetScreenAlignPartial() const;
-
 	// If partial, arc angle of partial torus in radians.
 	virtual void SetAngle(float InAngle);
 	virtual float GetAngle() const;
@@ -95,7 +96,7 @@ protected:
 
 	// Torus inner radius.
 	UPROPERTY()
-	float InnerRadius = 100.0f;
+	float InnerRadius = 5.0f;
 
 	// Number of segments for rendering torus.
 	UPROPERTY()
