Add missing setter/getters for hit owner flag.

diff --git a/Engine/Source/Runtime/InteractiveToolsFramework/Private/BaseGizmos/GizmoElementGroup.cpp b/Engine/Source/Runtime/InteractiveToolsFramework/Private/BaseGizmos/GizmoElementGroup.cpp
--- a/Engine/Source/Runtime/InteractiveToolsFramework/Private/BaseGizmos/GizmoElementGroup.cpp
+++ b/Engine/Source/Runtime/InteractiveToolsFramework/Private/BaseGizmos/GizmoElementGroup.cpp
@@ -151,4 +151,14 @@ void UGizmoElementGroup::SetConstantScale(bool bInConstantScale)
 bool UGizmoElementGroup::GetConstantScale() const
 {
 	return bConstantScale;
-}
\ No newline at end of file
+}
+
+void UGizmoElementGroup::SetHitOwner(bool bInHitOwner)
+{
+	bHitOwner = bInHitOwner;
+}
+
+bool UGizmoElementGroup::GetHitOwner() const
+{
+	return bHitOwner;
+}
diff --git a/Engine/Source/Runtime/InteractiveToolsFramework/Public/BaseGizmos/GizmoElementGroup.h b/Engine/Source/Runtime/InteractiveToolsFramework/Public/BaseGizmos/GizmoElementGroup.h
--- a/Engine/Source/Runtime/InteractiveToolsFramework/Public/BaseGizmos/GizmoElementGroup.h
+++ b/Engine/Source/Runtime/InteractiveToolsFramework/Public/BaseGizmos/GizmoElementGroup.h
@@ -44,6 +44,10 @@ public:
 	virtual void SetConstantScale(bool InConstantScale);
 	virtual bool GetConstantScale() const;
 
+	// Set whether this group should be treated as a hit owner and its part identifier returned when any of its sub-elements are hit.
+	virtual void SetHitOwner(bool bInHitOwner);
+	virtual bool GetHitOwner() const;
+
 protected:
 
 	// When true, maintains view-dependent constant scale for this gizmo object hierarchy
