New TRS gizmo: fix line scaling bug in line-drawn handles.

diff --git a/Engine/Source/Runtime/InteractiveToolsFramework/Private/BaseGizmos/GizmoElementCircle.cpp b/Engine/Source/Runtime/InteractiveToolsFramework/Private/BaseGizmos/GizmoElementCircle.cpp
--- a/Engine/Source/Runtime/InteractiveToolsFramework/Private/BaseGizmos/GizmoElementCircle.cpp
+++ b/Engine/Source/Runtime/InteractiveToolsFramework/Private/BaseGizmos/GizmoElementCircle.cpp
@@ -36,7 +36,12 @@ void UGizmoElementCircle::Render(IToolsContextRenderAPI* RenderAPI, const FRende
 
 		if (bDrawLine)
 		{
-			DrawCircle(PDI, WorldCenter, WorldUpAxis, WorldSideAxis, CurrentRenderState.GetCurrentLineColor(), WorldRadius, NumSides, SDPG_Foreground, GetCurrentLineThickness());
+			check(RenderAPI);
+			const FSceneView* View = RenderAPI->GetSceneView();
+			check(View);
+			float CurrentLineThickness = GetCurrentLineThickness(View->IsPerspectiveProjection(), View->FOV);
+
+			DrawCircle(PDI, WorldCenter, WorldUpAxis, WorldSideAxis, CurrentRenderState.GetCurrentLineColor(), WorldRadius, NumSides, SDPG_Foreground, CurrentLineThickness, 0.0, bScreenSpaceLine);
 		}
 	}
 }
diff --git a/Engine/Source/Runtime/InteractiveToolsFramework/Private/BaseGizmos/GizmoElementLineBase.cpp b/Engine/Source/Runtime/InteractiveToolsFramework/Private/BaseGizmos/GizmoElementLineBase.cpp
--- a/Engine/Source/Runtime/InteractiveToolsFramework/Private/BaseGizmos/GizmoElementLineBase.cpp
+++ b/Engine/Source/Runtime/InteractiveToolsFramework/Private/BaseGizmos/GizmoElementLineBase.cpp
@@ -11,18 +11,29 @@ bool UGizmoElementLineBase::UpdateRenderState(IToolsContextRenderAPI* RenderAPI,
 	return Super::UpdateRenderState(RenderAPI, InLocalOrigin, InOutRenderState);
 }
 
-float UGizmoElementLineBase::GetCurrentLineThickness() const
+float UGizmoElementLineBase::GetCurrentLineThickness(bool bPerspectiveView, float InViewFOV) const
 {
+	float CurrentLineThickness;
+
 	if (ElementInteractionState == EGizmoElementInteractionState::Hovering)
 	{
-		return (LineThickness > 0.0f ? LineThickness * HoverLineThicknessMultiplier : HoverLineThicknessMultiplier);
+		CurrentLineThickness = (LineThickness > 0.0f ? LineThickness * HoverLineThicknessMultiplier : HoverLineThicknessMultiplier);
 	}
 	else if (ElementInteractionState == EGizmoElementInteractionState::Interacting)
 	{
-		return (LineThickness > 0.0f ? LineThickness * InteractLineThicknessMultiplier : InteractLineThicknessMultiplier);
+		CurrentLineThickness = (LineThickness > 0.0f ? LineThickness * InteractLineThicknessMultiplier : InteractLineThicknessMultiplier);
+	}
+	else
+	{
+		CurrentLineThickness = LineThickness;
 	}
 
-	return LineThickness;
+	if (bPerspectiveView)
+	{
+		CurrentLineThickness *= (InViewFOV / 90.0);		// compensate for FOV scaling in Gizmos...
+	}
+
+	return CurrentLineThickness;
 }
 
 void UGizmoElementLineBase::SetLineThickness(float InLineThickness)
@@ -63,6 +74,16 @@ float UGizmoElementLineBase::GetInteractLineThicknessMultiplier() const
 	return InteractLineThicknessMultiplier;
 }
 
+void UGizmoElementLineBase::SetScreenSpaceLine(bool bInScreenSpaceLine)
+{
+	bScreenSpaceLine = bInScreenSpaceLine;
+}
+
+bool UGizmoElementLineBase::GetScreenSpaceLine() const
+{
+	return bScreenSpaceLine;
+}
+
 void UGizmoElementLineBase::SetLineColor(FLinearColor InLineColor, bool InOverridesChildState)
 {
 	LineRenderAttributes.LineColor.SetColor(InLineColor, InOverridesChildState);
diff --git a/Engine/Source/Runtime/InteractiveToolsFramework/Private/BaseGizmos/GizmoElementRectangle.cpp b/Engine/Source/Runtime/InteractiveToolsFramework/Private/BaseGizmos/GizmoElementRectangle.cpp
--- a/Engine/Source/Runtime/InteractiveToolsFramework/Private/BaseGizmos/GizmoElementRectangle.cpp
+++ b/Engine/Source/Runtime/InteractiveToolsFramework/Private/BaseGizmos/GizmoElementRectangle.cpp
@@ -33,8 +33,13 @@ void UGizmoElementRectangle::Render(IToolsContextRenderAPI* RenderAPI, const FRe
 		}
 		if (bDrawLine)
 		{
+			check(RenderAPI);
+			const FSceneView* View = RenderAPI->GetSceneView();
+			check(View);
+			float CurrentLineThickness = GetCurrentLineThickness(View->IsPerspectiveProjection(), View->FOV);
+
 			FColor LineColor = CurrentRenderState.GetCurrentLineColor().ToFColor(false);
-			DrawRectangle(PDI, WorldCenter, WorldUpAxis, WorldSideAxis, LineColor, WorldWidth, WorldHeight, SDPG_Foreground, GetCurrentLineThickness());
+			DrawRectangle(PDI, WorldCenter, WorldUpAxis, WorldSideAxis, LineColor, WorldWidth, WorldHeight, SDPG_Foreground, CurrentLineThickness, 0.0, bScreenSpaceLine);
 		}
 	}
 }
diff --git a/Engine/Source/Runtime/InteractiveToolsFramework/Public/BaseGizmos/GizmoElementLineBase.h b/Engine/Source/Runtime/InteractiveToolsFramework/Public/BaseGizmos/GizmoElementLineBase.h
--- a/Engine/Source/Runtime/InteractiveToolsFramework/Public/BaseGizmos/GizmoElementLineBase.h
+++ b/Engine/Source/Runtime/InteractiveToolsFramework/Public/BaseGizmos/GizmoElementLineBase.h
@@ -19,8 +19,8 @@ class INTERACTIVETOOLSFRAMEWORK_API UGizmoElementLineBase : public UGizmoElement
 	GENERATED_BODY()
 public:
 
-	// Get line thickness for based on current element interaction state
-	virtual float GetCurrentLineThickness() const;
+	// Get line thickness for based on current element interaction state and view
+	virtual float GetCurrentLineThickness(bool bPerspectiveView, float InViewFOV) const;
 
 	// Line thickness when rendering lines, 0.0 is valid and will render thinnest line 
 	virtual void SetLineThickness(float InLineThickness);
@@ -34,6 +34,10 @@ public:
 	virtual void SetInteractLineThicknessMultiplier(float InInteractLineThicknessMultiplier);
 	virtual float GetInteractLineThicknessMultiplier() const;
 
+	// Whether line thickness is in screen space 
+	virtual void SetScreenSpaceLine(bool bInScreenSpaceLine);
+	virtual bool GetScreenSpaceLine() const;
+
 	//
 	// Methods for managing line state attributes: LineColor, HoverLineColor, InteractLineColor
 	//
@@ -104,6 +108,10 @@ protected:
 	UPROPERTY()
 	float LineThickness = 0.0;
 
+	// Whether line thickness is in screen space
+	UPROPERTY()
+	bool bScreenSpaceLine = true;
+
 	// Multiplier applied to line thickness when hovering
 	UPROPERTY(EditAnywhere, Category = Options)
 	float HoverLineThicknessMultiplier = 2.0f;
