Gizmo element library: add getter methods for part-specific state, needed for multi-part state updates used for highlighting gizmo parts in new TRS gizmo uniform scale and planar scale modes.

diff --git a/Engine/Source/Runtime/InteractiveToolsFramework/Private/BaseGizmos/GizmoElementBase.cpp b/Engine/Source/Runtime/InteractiveToolsFramework/Private/BaseGizmos/GizmoElementBase.cpp
--- a/Engine/Source/Runtime/InteractiveToolsFramework/Private/BaseGizmos/GizmoElementBase.cpp
+++ b/Engine/Source/Runtime/InteractiveToolsFramework/Private/BaseGizmos/GizmoElementBase.cpp
@@ -398,34 +398,57 @@ EGizmoElementInteractionState UGizmoElementBase::GetElementInteractionState() co
 	return ElementInteractionState;
 }
 
-void UGizmoElementBase::UpdatePartHittableState(bool bHittable, uint32 InPartIdentifier)
+bool UGizmoElementBase::UpdatePartHittableState(bool bHittable, uint32 InPartIdentifier)
 {
-	if (InPartIdentifier == PartIdentifier)
+	if (InPartIdentifier != PartIdentifier)
 	{
-		uint8 State = static_cast<uint8>(ElementState);
-		uint8 HittableMask = static_cast<uint8>(EGizmoElementState::Hittable);
-		uint8 NewState = (bHittable ? State | HittableMask : State & (~HittableMask));
-		ElementState = static_cast<EGizmoElementState>(NewState);
+		return false;
 	}
+
+	ElementState = (bHittable ? ElementState | EGizmoElementState::Hittable : ElementState & (~EGizmoElementState::Hittable));
+	return true;
 }
 
-void UGizmoElementBase::UpdatePartVisibleState(bool bVisible, uint32 InPartIdentifier)
+TOptional<bool> UGizmoElementBase::GetPartHittableState(uint32 InPartIdentifier) const
 {
-	if (InPartIdentifier == PartIdentifier)
+	return InPartIdentifier == PartIdentifier ? TOptional<bool>(GetHittableState()) : TOptional<bool>();
+}
+
+bool UGizmoElementBase::UpdatePartVisibleState(bool bVisible, uint32 InPartIdentifier)
+{
+	if (InPartIdentifier != PartIdentifier)
 	{
-		uint8 State = static_cast<uint8>(ElementState);
-		uint8 VisibleMask = static_cast<uint8>(EGizmoElementState::Visible);
-		uint8 NewState = (bVisible ? State | VisibleMask : State & (~VisibleMask));
-		ElementState = static_cast<EGizmoElementState>(NewState);
+		return false;
 	}
+
+	ElementState = (bVisible ? ElementState | EGizmoElementState::Visible : ElementState & (~EGizmoElementState::Visible));
+	return true;
 }
 
-void UGizmoElementBase::UpdatePartInteractionState(EGizmoElementInteractionState InInteractionState, uint32 InPartIdentifier)
+TOptional<bool> UGizmoElementBase::GetPartVisibleState(uint32 InPartIdentifier) const
+{
+	return InPartIdentifier == PartIdentifier ? TOptional<bool>(GetVisibleState()) : TOptional<bool>();
+}
+
+bool UGizmoElementBase::UpdatePartInteractionState(EGizmoElementInteractionState InInteractionState, uint32 InPartIdentifier)
+{
+	if (InPartIdentifier != PartIdentifier)
+	{
+		return false;
+	}
+	
+	ElementInteractionState = InInteractionState;
+	return true;
+}
+
+TOptional<EGizmoElementInteractionState> UGizmoElementBase::GetPartInteractionState(uint32 InPartIdentifier) const
 {
 	if (InPartIdentifier == PartIdentifier)
 	{
-		ElementInteractionState = InInteractionState;
+		return TOptional<EGizmoElementInteractionState>(GetElementInteractionState());
 	}
+
+	return  TOptional<EGizmoElementInteractionState>();
 }
 
 void UGizmoElementBase::SetViewDependentType(EGizmoElementViewDependentType InViewDependentType)
diff --git a/Engine/Source/Runtime/InteractiveToolsFramework/Private/BaseGizmos/GizmoElementGroup.cpp b/Engine/Source/Runtime/InteractiveToolsFramework/Private/BaseGizmos/GizmoElementGroup.cpp
--- a/Engine/Source/Runtime/InteractiveToolsFramework/Private/BaseGizmos/GizmoElementGroup.cpp
+++ b/Engine/Source/Runtime/InteractiveToolsFramework/Private/BaseGizmos/GizmoElementGroup.cpp
@@ -104,43 +104,129 @@ void UGizmoElementGroup::Remove(UGizmoElementBase* InElement)
 	}
 }
 
-void UGizmoElementGroup::UpdatePartVisibleState(bool bVisible, uint32 InPartIdentifier)
+bool UGizmoElementGroup::UpdatePartVisibleState(bool bVisible, uint32 InPartIdentifier)
 {
-	Super::UpdatePartVisibleState(bVisible, InPartIdentifier);
+	if (Super::UpdatePartVisibleState(bVisible, InPartIdentifier))
+	{
+		return true;
+	}
 
 	for (UGizmoElementBase* Element : Elements)
 	{
-		if (Element)
+		if (Element && Element->UpdatePartVisibleState(bVisible, InPartIdentifier))
 		{
-			Element->UpdatePartVisibleState(bVisible, InPartIdentifier);
+			return true;
 		}
 	}
+
+	return false;
 }
 
-void UGizmoElementGroup::UpdatePartHittableState(bool bHittable, uint32 InPartIdentifier)
+TOptional<bool> UGizmoElementGroup::GetPartVisibleState(uint32 InPartIdentifier) const
 {
-	Super::UpdatePartHittableState(bHittable, InPartIdentifier);
+	TOptional<bool> Result = Super::GetPartVisibleState(InPartIdentifier);
+	if (Result.IsSet())
+	{
+		return Result;
+	}
 
 	for (UGizmoElementBase* Element : Elements)
 	{
 		if (Element)
 		{
-			Element->UpdatePartHittableState(bHittable, InPartIdentifier);
+			Result = Element->GetPartVisibleState(InPartIdentifier);
+			if (Result.IsSet())
+			{
+				return Result;
+			}
 		}
 	}
+
+	return TOptional<bool>();
 }
 
-void UGizmoElementGroup::UpdatePartInteractionState(EGizmoElementInteractionState InInteractionState, uint32 InPartIdentifier)
+
+bool UGizmoElementGroup::UpdatePartHittableState(bool bHittable, uint32 InPartIdentifier)
 {
-	Super::UpdatePartInteractionState(InInteractionState, InPartIdentifier);
+	if (Super::UpdatePartHittableState(bHittable, InPartIdentifier))
+	{
+		return true;
+	}
+
+	for (UGizmoElementBase* Element : Elements)
+	{
+		if (Element && Element->UpdatePartHittableState(bHittable, InPartIdentifier))
+		{
+			return true;
+		}
+	}
+
+	return false;
+}
+
+TOptional<bool> UGizmoElementGroup::GetPartHittableState(uint32 InPartIdentifier) const
+{
+	TOptional<bool> Result = Super::GetPartHittableState(InPartIdentifier);
+	if (Result.IsSet())
+	{
+		return Result;
+	}
 
 	for (UGizmoElementBase* Element : Elements)
 	{
 		if (Element)
 		{
-			Element->UpdatePartInteractionState(InInteractionState, InPartIdentifier);
+			Result = Element->GetPartHittableState(InPartIdentifier);
+			if (Result.IsSet())
+			{
+				return Result;
+			}
+		}
+	}
+
+	return TOptional<bool>();
+}
+
+bool UGizmoElementGroup::UpdatePartInteractionState(EGizmoElementInteractionState InInteractionState, uint32 InPartIdentifier)
+{
+	if (Super::UpdatePartInteractionState(InInteractionState, InPartIdentifier))
+	{
+		return true;
+	}
+
+	for (UGizmoElementBase* Element : Elements)
+	{
+		if (Element && Element->UpdatePartInteractionState(InInteractionState, InPartIdentifier))
+		{
+			return true;
 		}
 	}
+
+	return false;
+}
+
+
+TOptional<EGizmoElementInteractionState> UGizmoElementGroup::GetPartInteractionState(uint32 InPartIdentifier) const
+{
+	TOptional<EGizmoElementInteractionState> Result = Super::GetPartInteractionState(InPartIdentifier);
+	if (Result.IsSet())
+	{
+		return Result;
+	}
+
+	for (UGizmoElementBase* Element : Elements)
+	{
+		if (Element)
+		{ 
+			Result = Element->GetPartInteractionState(InPartIdentifier);
+			if (Result.IsSet())
+			{
+				return Result;
+			}
+		}
+	}
+
+	return TOptional<EGizmoElementInteractionState>();
 }
 
 void UGizmoElementGroup::SetConstantScale(bool bInConstantScale)
diff --git a/Engine/Source/Runtime/InteractiveToolsFramework/Public/BaseGizmos/GizmoElementBase.h b/Engine/Source/Runtime/InteractiveToolsFramework/Public/BaseGizmos/GizmoElementBase.h
--- a/Engine/Source/Runtime/InteractiveToolsFramework/Public/BaseGizmos/GizmoElementBase.h
+++ b/Engine/Source/Runtime/InteractiveToolsFramework/Public/BaseGizmos/GizmoElementBase.h
@@ -169,14 +169,23 @@ public:
 	virtual void SetElementInteractionState(EGizmoElementInteractionState InInteractionState);
 	virtual EGizmoElementInteractionState GetElementInteractionState() const;
 
-	// Update element's visibility state if element is associated with the specified gizmo part.
-	virtual void UpdatePartVisibleState(bool bVisible, uint32 InPartIdentifier);
+	// Update element's visibility state if element is associated with the specified gizmo part, return true if part was found.
+	virtual bool UpdatePartVisibleState(bool bVisible, uint32 InPartIdentifier);
 
-	// Update element's hittable state if element is associated with the specified gizmo part.
-	virtual void UpdatePartHittableState(bool bHittable, uint32 InPartIdentifier);
+	// Get element's visible state for element associated with the specified gizmo part, if part id was found.
+	virtual TOptional<bool> GetPartVisibleState(uint32 InPartIdentifier) const;
 
-	// Update element's interaction state if element is associated with the specified gizmo part.
-	virtual void UpdatePartInteractionState(EGizmoElementInteractionState InInteractionState, uint32 InPartIdentifier);
+	// Update element's hittable state if element is associated with the specified gizmo part, return true if part id was found.
+	virtual bool UpdatePartHittableState(bool bHittable, uint32 InPartIdentifier);
+
+	// Get element's hittable state for element associated with the specified gizmo part, if part id was found.
+	virtual TOptional<bool> GetPartHittableState(uint32 InPartIdentifier) const;
+
+	// Update element's interaction state if element is associated with the specified gizmo part, return true if part id was found.
+	virtual bool UpdatePartInteractionState(EGizmoElementInteractionState InInteractionState, uint32 InPartIdentifier);
+
+	// Get element's interaction state for element associated with the specified gizmo part, if part id was found.
+	virtual TOptional<EGizmoElementInteractionState> GetPartInteractionState(uint32 InPartIdentifier) const;
 
 	// View-dependent type - None, Axis or Plane. 
 	virtual void SetViewDependentType(EGizmoElementViewDependentType ViewDependentType);
diff --git a/Engine/Source/Runtime/InteractiveToolsFramework/Public/BaseGizmos/GizmoElementGroup.h b/Engine/Source/Runtime/InteractiveToolsFramework/Public/BaseGizmos/GizmoElementGroup.h
--- a/Engine/Source/Runtime/InteractiveToolsFramework/Public/BaseGizmos/GizmoElementGroup.h
+++ b/Engine/Source/Runtime/InteractiveToolsFramework/Public/BaseGizmos/GizmoElementGroup.h
@@ -31,14 +31,23 @@ public:
 	// Remove object from group, if it exists
 	virtual void Remove(UGizmoElementBase* InElement);
 
-	// Update group and contained elements' visibility state for elements in specified gizmo parts.
-	virtual void UpdatePartVisibleState(bool bVisible, uint32 InPartIdentifier);
+	// Update group and contained elements' visibility state for elements in specified gizmo parts, return true if part was found.
+	virtual bool UpdatePartVisibleState(bool bVisible, uint32 InPartIdentifier);
 
-	// Update group and contained elements' hittable state for elements in specified gizmo parts.
-	virtual void UpdatePartHittableState(bool bHittable, uint32 InPartIdentifier);
+	// Get element's visible state for element associated with the specified gizmo part, if part was found.
+	virtual TOptional<bool> GetPartVisibleState(uint32 InPartIdentifier) const;
 
-	// Update group and contained elements' interaction state for elements in specified gizmo parts.
-	virtual void UpdatePartInteractionState(EGizmoElementInteractionState InInteractionState, uint32 InPartIdentifier);
+	// Update group and contained elements' hittable state for elements in specified gizmo parts, return true if part was found.
+	virtual bool UpdatePartHittableState(bool bHittable, uint32 InPartIdentifier);
+
+	// Get element's hittable state for element associated with the specified gizmo part, if part was found.
+	virtual TOptional<bool> GetPartHittableState(uint32 InPartIdentifier) const;
+
+	// Update group and contained elements' interaction state for elements in specified gizmo parts, return true if part was found.
+	virtual bool UpdatePartInteractionState(EGizmoElementInteractionState InInteractionState, uint32 InPartIdentifier);
+
+	// Get element's interaction state for element associated with the specified gizmo part, if part was found.
+	virtual TOptional<EGizmoElementInteractionState> GetPartInteractionState(uint32 InPartIdentifiero) const;
 
 	// When true, maintains view-dependent constant scale for this gizmo object hierarchy
 	virtual void SetConstantScale(bool InConstantScale);
diff --git a/Engine/Source/Runtime/InteractiveToolsFramework/Public/BaseGizmos/GizmoElementShared.h b/Engine/Source/Runtime/InteractiveToolsFramework/Public/BaseGizmos/GizmoElementShared.h
--- a/Engine/Source/Runtime/InteractiveToolsFramework/Public/BaseGizmos/GizmoElementShared.h
+++ b/Engine/Source/Runtime/InteractiveToolsFramework/Public/BaseGizmos/GizmoElementShared.h
@@ -17,6 +17,8 @@ enum class EGizmoElementState : uint8
 	VisibleAndHittable = Visible | Hittable
 };
 
+ENUM_CLASS_FLAGS(EGizmoElementState)
+
 //
 // Interaction state of gizmo element
 //
