make gizmomath do computations in double precision
+ make PlanePositionGizmo not crash if plane intersection fails

diff --git a/Engine/Source/Runtime/InteractiveToolsFramework/Private/BaseGizmos/GizmoMath.cpp b/Engine/Source/Runtime/InteractiveToolsFramework/Private/BaseGizmos/GizmoMath.cpp
--- a/Engine/Source/Runtime/InteractiveToolsFramework/Private/BaseGizmos/GizmoMath.cpp
+++ b/Engine/Source/Runtime/InteractiveToolsFramework/Private/BaseGizmos/GizmoMath.cpp
@@ -9,7 +9,7 @@ FVector GizmoMath::ProjectPointOntoLine(
 	const FVector& Point,
 	const FVector& LineOrigin, const FVector& LineDirection)
 {
-	float ProjectionParam = FVector::DotProduct((Point - LineOrigin), LineDirection);
+	double ProjectionParam = FVector::DotProduct((Point - LineOrigin), LineDirection);
 	return LineOrigin + ProjectionParam * LineDirection;
 }
 
@@ -20,8 +20,9 @@ void GizmoMath::NearestPointOnLine(
 	FVector& NearestPointOut, float& LineParameterOut)
 {
 	check(LineDirection.IsNormalized());
-	LineParameterOut = FVector::DotProduct( (QueryPoint - LineOrigin), LineDirection);
-	NearestPointOut = LineOrigin + LineParameterOut * LineDirection;
+	double LineParameter = FVector::DotProduct( (QueryPoint - LineOrigin), LineDirection);
+	NearestPointOut = LineOrigin + LineParameter * LineDirection;
+	LineParameterOut = (float)LineParameter;
 }
 
 
@@ -32,21 +33,21 @@ void GizmoMath::NearestPointOnLineToRay(
 	FVector& NearestRayPointOut, float& RayParameterOut)
 {
 	FVector kDiff = LineOrigin - RayOrigin;
-	float a01 = -FVector::DotProduct(LineDirection, RayDirection);
-	float b0 = FVector::DotProduct(kDiff, LineDirection);
-	float c = kDiff.SizeSquared();
-	float det = FMath::Abs((float)1 - a01 * a01);
-	float b1, s0, s1;
+	double a01 = -FVector::DotProduct(LineDirection, RayDirection);
+	double b0 = FVector::DotProduct(kDiff, LineDirection);
+	double c = kDiff.SizeSquared();
+	double det = FMath::Abs((double)1 - a01 * a01);
+	double b1, s0, s1;
 
 	if (det >= SMALL_NUMBER)
 	{
 		b1 = -FVector::DotProduct(kDiff, RayDirection);
 		s1 = a01 * b0 - b1;
 
-		if (s1 >= (float)0)
+		if (s1 >= (double)0)
 		{
 			// Two interior points are closest, one on Line and one on Ray
-			float invDet = ((float)1) / det;
+			double invDet = ((double)1) / det;
 			s0 = (a01 * b1 - b0) * invDet;
 			s1 *= invDet;
 		}
@@ -54,20 +55,20 @@ void GizmoMath::NearestPointOnLineToRay(
 		{
 			// Origin of Ray and interior point of Line are closest.
 			s0 = -b0;
-			s1 = (float)0;
+			s1 = (double)0;
 		}
 	}
 	else
 	{
 		// Lines are parallel, closest pair with one point at Ray origin.
 		s0 = -b0;
-		s1 = (float)0;
+		s1 = (double)0;
 	}
 
 	NearestLinePointOut = LineOrigin + s0 * LineDirection;
 	NearestRayPointOut = RayOrigin + s1 * RayDirection;
-	LineParameterOut = s0;
-	RayParameterOut = s1;
+	LineParameterOut = (float)s0;
+	RayParameterOut = (float)s1;
 }
 
 
@@ -81,15 +82,15 @@ void GizmoMath::RayPlaneIntersectionPoint(
 	bIntersectsOut = false;
 	PlaneIntersectionPointOut = PlaneOrigin;
 
-	float PlaneEquationD = -FVector::DotProduct(PlaneOrigin, PlaneNormal);
-	float NormalDot = FVector::DotProduct(RayDirection, PlaneNormal);
+	double PlaneEquationD = -FVector::DotProduct(PlaneOrigin, PlaneNormal);
+	double NormalDot = FVector::DotProduct(RayDirection, PlaneNormal);
 
 	if (FMath::Abs(NormalDot) < SMALL_NUMBER)
 	{
 		return;
 	}
 
-	float RayParam = -( FVector::DotProduct(RayOrigin, PlaneNormal) + PlaneEquationD) / NormalDot;
+	double RayParam = -( FVector::DotProduct(RayOrigin, PlaneNormal) + PlaneEquationD) / NormalDot;
 	if (RayParam < 0)
 	{
 		return;
@@ -109,16 +110,16 @@ void GizmoMath::RaySphereIntersection(
 	SphereIntersectionPointOut = RayOrigin;
 	 
 	FVector DeltaPos = RayOrigin - SphereOrigin;
-	float a0 = DeltaPos.SizeSquared() - SphereRadius*SphereRadius;
-	float a1 = FVector::DotProduct(RayDirection, DeltaPos);
-	float discr = a1 * a1 - a0;
+	double a0 = DeltaPos.SizeSquared() - (double)SphereRadius*(double)SphereRadius;
+	double a1 = FVector::DotProduct(RayDirection, DeltaPos);
+	double discr = a1 * a1 - a0;
 	if (discr > 0)   // intersection only when roots are real
 	{
 		bIntersectsOut = true;
-		float root = FMath::Sqrt(discr);
-		float NearRayParam = -a1 + root;		// isn't it always this one?
-		float NearRayParam2 = -a1 - root;
-		float UseRayParam = FMath::Min(NearRayParam, NearRayParam2);
+		double root = FMath::Sqrt(discr);
+		double NearRayParam = -a1 + root;		// isn't it always this one?
+		double NearRayParam2 = -a1 - root;
+		double UseRayParam = FMath::Min(NearRayParam, NearRayParam2);
 		SphereIntersectionPointOut = RayOrigin + UseRayParam * RayDirection;
 	}
 }
@@ -604,16 +605,16 @@ void GizmoMath::ClosetPointOnCircle(
 {
 	FVector PointDelta = QueryPoint - CircleOrigin;
 	FVector DeltaInPlane = PointDelta - FVector::DotProduct(CircleNormal,PointDelta)*CircleNormal;
-	float OriginDist = DeltaInPlane.Size();
+	double OriginDist = DeltaInPlane.Size();
 	if (OriginDist > 0.0f)
 	{
-		ClosestPointOut =  CircleOrigin + (CircleRadius / OriginDist) * DeltaInPlane;
+		ClosestPointOut =  CircleOrigin + ((double)CircleRadius / OriginDist) * DeltaInPlane;
 	}
 	else    // all points equidistant, use any one
 	{
 		FVector PlaneX, PlaneY;
 		MakeNormalPlaneBasis(CircleNormal, PlaneX, PlaneY);
-		ClosestPointOut = CircleOrigin + CircleRadius * PlaneX;
+		ClosestPointOut = CircleOrigin + (double)CircleRadius * PlaneX;
 	}
 }
 
@@ -626,8 +627,8 @@ void GizmoMath::MakeNormalPlaneBasis(
 	// Duff et al method, from https://graphics.pixar.com/library/OrthonormalB/paper.pdf
 	if (PlaneNormal.Z < 0)
 	{
-		float A = 1.0f / (1.0f - PlaneNormal.Z);
-		float B = PlaneNormal.X * PlaneNormal.Y * A;
+		double A = 1.0f / (1.0f - PlaneNormal.Z);
+		double B = PlaneNormal.X * PlaneNormal.Y * A;
 		BasisAxis1Out.X = 1.0f - PlaneNormal.X * PlaneNormal.X * A;
 		BasisAxis1Out.Y = -B;
 		BasisAxis1Out.Z = PlaneNormal.X;
@@ -637,8 +638,8 @@ void GizmoMath::MakeNormalPlaneBasis(
 	}
 	else
 	{
-		float A = 1.0f / (1.0f + PlaneNormal.Z);
-		float B = -PlaneNormal.X * PlaneNormal.Y * A;
+		double A = 1.0f / (1.0f + PlaneNormal.Z);
+		double B = -PlaneNormal.X * PlaneNormal.Y * A;
 		BasisAxis1Out.X = 1.0f - PlaneNormal.X * PlaneNormal.X * A;
 		BasisAxis1Out.Y = B;
 		BasisAxis1Out.Z = -PlaneNormal.X;
@@ -658,8 +659,8 @@ float GizmoMath::ComputeAngleInPlane(
 	// project point into plane
 	FVector LocalPoint = Point - PlaneOrigin;
 
-	float X = FVector::DotProduct(LocalPoint, PlaneAxis1);
-	float Y = FVector::DotProduct(LocalPoint, PlaneAxis2);
+	double X = FVector::DotProduct(LocalPoint, PlaneAxis1);
+	double Y = FVector::DotProduct(LocalPoint, PlaneAxis2);
 
 	float SignedAngle = (float)atan2(Y, X);
 	return SignedAngle;
@@ -674,8 +675,8 @@ FVector2D GizmoMath::ComputeCoordinatesInPlane(
 	const FVector& PlaneAxis1, const FVector& PlaneAxis2)
 {
 	FVector LocalPoint = Point - PlaneOrigin;
-	float X = FVector::DotProduct(LocalPoint, PlaneAxis1);
-	float Y = FVector::DotProduct(LocalPoint, PlaneAxis2);
+	double X = FVector::DotProduct(LocalPoint, PlaneAxis1);
+	double Y = FVector::DotProduct(LocalPoint, PlaneAxis2);
 	return FVector2D(X, Y);
 }
 
@@ -685,27 +686,27 @@ FVector GizmoMath::ProjectPointOntoPlane(
 	const FVector& PlaneOrigin, const FVector& PlaneNormal)
 {
 	FVector LocalPoint = Point - PlaneOrigin;
-	float NormalDot = FVector::DotProduct(LocalPoint, PlaneNormal);
+	double NormalDot = FVector::DotProduct(LocalPoint, PlaneNormal);
 	return Point - NormalDot * PlaneNormal;
 }
 
 
-
-float GizmoMath::SnapToIncrement(float Value, float Increment)
+template <typename RealType>
+RealType GizmoMath::SnapToIncrement(RealType Value, RealType Increment)
 {
 	if (!FMath::IsFinite(Value))
 	{
 		return 0;
 	}
-	float Sign = FMath::Sign(Value);
+	RealType Sign = FMath::Sign(Value);
 	Value = FMath::Abs(Value);
 	int IntIncrement = (int)(Value / Increment);
-	float Remainder = (float)fmod(Value, Increment);
+	RealType Remainder = (RealType)fmod(Value, Increment);
 	if (Remainder > IntIncrement / 2)
 	{
 		++IntIncrement;
 	}
-	return Sign * (float)IntIncrement * Increment;
+	return Sign * (RealType)IntIncrement * Increment;
 }
 
 // @todo: Remove this and replace calls to it with MakePerpVectors() once
@@ -752,5 +753,10 @@ namespace GizmoMath
 		const FVector& ConeCenter, const FVector& ConeDirection, double ConeAngle, double ConeHeight,
 		const FVector& RayOrigin, const FVector& RayDirection,
 		bool& bIntersectsOut, double& OutHitDepth);
+	
+	template
+	double INTERACTIVETOOLSFRAMEWORK_API GizmoMath::SnapToIncrement(double Value, double Increment);
 
+	template
+	float INTERACTIVETOOLSFRAMEWORK_API GizmoMath::SnapToIncrement(float Value, float Increment);
 } // end namespace GizmoMath
diff --git a/Engine/Source/Runtime/InteractiveToolsFramework/Private/BaseGizmos/PlanePositionGizmo.cpp b/Engine/Source/Runtime/InteractiveToolsFramework/Private/BaseGizmos/PlanePositionGizmo.cpp
--- a/Engine/Source/Runtime/InteractiveToolsFramework/Private/BaseGizmos/PlanePositionGizmo.cpp
+++ b/Engine/Source/Runtime/InteractiveToolsFramework/Private/BaseGizmos/PlanePositionGizmo.cpp
@@ -71,7 +71,12 @@ void UPlanePositionGizmo::OnClickPress(const FInputDeviceRay& PressPos)
 		InteractionOrigin, InteractionNormal,
 		PressPos.WorldRay.Origin, PressPos.WorldRay.Direction,
 		bIntersects, IntersectionPoint);
-	check(bIntersects);  // need to handle this case...
+	if (!bIntersects)
+	{
+		// Generally should not happen since the user clicked the plane to start the interaction, but could happen in a floating point error edge case
+		bInInteraction = false;
+		return;
+	}
 
 	InteractionStartPoint = InteractionCurPoint = IntersectionPoint;
 
@@ -110,6 +115,11 @@ void UPlanePositionGizmo::OnClickPress(const FInputDeviceRay& PressPos)
 
 void UPlanePositionGizmo::OnClickDrag(const FInputDeviceRay& DragPos)
 {
+	if (!bInInteraction)
+	{
+		return;
+	}
+
 	FVector HitPoint;
 	FVector2D NewParamValue;
 
@@ -152,7 +162,10 @@ void UPlanePositionGizmo::OnClickDrag(const FInputDeviceRay& DragPos)
 
 void UPlanePositionGizmo::OnClickRelease(const FInputDeviceRay& ReleasePos)
 {
-	check(bInInteraction);
+	if (!bInInteraction)
+	{
+		return;
+	}
 
 	ParameterSource->EndModify();
 	if (StateTarget)
@@ -165,7 +178,10 @@ void UPlanePositionGizmo::OnClickRelease(const FInputDeviceRay& ReleasePos)
 
 void UPlanePositionGizmo::OnTerminateDragSequence()
 {
-	check(bInInteraction);
+	if (!bInInteraction)
+	{
+		return;
+	}
 
 	ParameterSource->EndModify();
 	if (StateTarget)
diff --git a/Engine/Source/Runtime/InteractiveToolsFramework/Public/BaseGizmos/GizmoMath.h b/Engine/Source/Runtime/InteractiveToolsFramework/Public/BaseGizmos/GizmoMath.h
--- a/Engine/Source/Runtime/InteractiveToolsFramework/Public/BaseGizmos/GizmoMath.h
+++ b/Engine/Source/Runtime/InteractiveToolsFramework/Public/BaseGizmos/GizmoMath.h
@@ -154,7 +154,8 @@ namespace GizmoMath
 	 * Round Value to nearest step of Increment
 	 * @return snapped/rounded value
 	 */
-	INTERACTIVETOOLSFRAMEWORK_API float SnapToIncrement(float Value, float Increment);
+	template <typename RealType>
+	INTERACTIVETOOLSFRAMEWORK_API RealType SnapToIncrement(RealType Value, RealType Increment);
 
 	/**
 	 * @todo change FVector to a templated vector type specialized by RealType (double or float)
