Fix torus intersection approximation and handle partial torii correctly

diff --git a/Engine/Source/Runtime/InteractiveToolsFramework/Private/BaseGizmos/GizmoElementTorus.cpp b/Engine/Source/Runtime/InteractiveToolsFramework/Private/BaseGizmos/GizmoElementTorus.cpp
--- a/Engine/Source/Runtime/InteractiveToolsFramework/Private/BaseGizmos/GizmoElementTorus.cpp
+++ b/Engine/Source/Runtime/InteractiveToolsFramework/Private/BaseGizmos/GizmoElementTorus.cpp
@@ -6,6 +6,7 @@
 #include "InputState.h"
 #include "Materials/MaterialInterface.h"
 #include "SceneManagement.h"
+#include "VectorUtil.h"
 
 void UGizmoElementTorus::Render(IToolsContextRenderAPI* RenderAPI, const FRenderTraversalState& RenderState)
 {
@@ -48,14 +49,19 @@ FInputRayHit UGizmoElementTorus::LineTrace(const UGizmoViewContext* ViewContext,
 		const double WorldInnerRadius = InnerRadius * CurrentLineTraceState.LocalToWorldTransform.GetScale3D().X;
 		const FVector WorldCenter = CurrentLineTraceState.LocalToWorldTransform.TransformPosition(FVector::ZeroVector);
 		const FVector WorldNormal = CurrentLineTraceState.LocalToWorldTransform.TransformVectorNoScale(Normal);
-		const FVector WorldBeginAxis = CurrentLineTraceState.LocalToWorldTransform.GetRotation().RotateVector(BeginAxis);
+		const FVector WorldBeginAxis = CurrentLineTraceState.LocalToWorldTransform.TransformVectorNoScale(BeginAxis);
 		double HitDepth = -1.0;
 
+		const FVector ViewToTorusCenter = 
+			ViewContext->IsPerspectiveProjection() ? (WorldCenter - ViewContext->ViewLocation).GetSafeNormal() : ViewContext->GetViewDirection();
 
 		// Determine if the ray direction is at a glancing angle by using a minimum cos angle based on the
 		// angle between the vector from torus center to ring center and the vector from torus center to ring edge.
-		double MinCosAngle = OuterRadius / FMath::Sqrt(OuterRadius * OuterRadius + InnerRadius * InnerRadius);
-		bool bAtGlancingAngle = FMath::Abs(FVector::DotProduct(WorldNormal, RayDirection)) <= MinCosAngle;
+		double MinCosAngle = 
+			FMath::Max(InnerRadius * 1.5 / FMath::Sqrt(OuterRadius * OuterRadius + InnerRadius * InnerRadius), 
+				DefaultViewAlignPlanarMinCosAngleTol);
+
+		bool bAtGlancingAngle = FMath::Abs(FVector::DotProduct(WorldNormal, ViewToTorusCenter)) <= MinCosAngle;
 
 		if (bAtGlancingAngle)
 		{
@@ -83,7 +89,7 @@ FInputRayHit UGizmoElementTorus::LineTrace(const UGizmoViewContext* ViewContext,
 				if (i > 0)
 				{
 					VectorA = VectorB;
-					VectorB = VectorA.RotateAngleAxisRad(AngleDelta, Normal);
+					VectorB = VectorA.RotateAngleAxisRad(AngleDelta, WorldNormal);
 				}
 
 				if (i == NumCylinders - 1)
@@ -139,6 +145,23 @@ FInputRayHit UGizmoElementTorus::LineTrace(const UGizmoViewContext* ViewContext,
 				return FInputRayHit();
 			}
 
+			if (bPartial)
+			{
+				// Compute projected angle
+				FVector HitVec = (NearestCirclePos - WorldCenter).GetSafeNormal();
+				double HitAngle = UE::Geometry::VectorUtil::PlaneAngleSignedR(WorldBeginAxis, HitVec, WorldNormal); 
+
+				if (HitAngle < 0)
+				{
+					HitAngle = UE_DOUBLE_TWO_PI - HitAngle;
+				}
+
+				if (HitAngle > Angle)
+				{
+					return FInputRayHit();
+				}
+			}
+
 			// Adjust hit depth to return the closest hit point's depth
 			HitDepth = (NearestRayPos - RayOrigin).Length();
 		}
