GizmoViewContext: copied over ScreenToPixel() and FMatrices::GetInvViewMatrix() methods from FSceneView.

diff --git a/Engine/Source/Runtime/InteractiveToolsFramework/Public/BaseGizmos/GizmoViewContext.h b/Engine/Source/Runtime/InteractiveToolsFramework/Public/BaseGizmos/GizmoViewContext.h
--- a/Engine/Source/Runtime/InteractiveToolsFramework/Public/BaseGizmos/GizmoViewContext.h
+++ b/Engine/Source/Runtime/InteractiveToolsFramework/Public/BaseGizmos/GizmoViewContext.h
@@ -28,14 +28,17 @@ public:
 		void ResetFromSceneView(const FSceneView& SceneView)
 		{
 			ViewMatrix = SceneView.ViewMatrices.GetViewMatrix();
+			InvViewMatrix = SceneView.ViewMatrices.GetInvViewMatrix();
 			ViewProjectionMatrix = SceneView.ViewMatrices.GetViewProjectionMatrix();
 		}
 
 		const FMatrix& GetViewMatrix() const { return ViewMatrix; }
+		const FMatrix& GetInvViewMatrix() const { return InvViewMatrix; }
 		const FMatrix& GetViewProjectionMatrix() const { return ViewProjectionMatrix; }
 
 	protected:
 		FMatrix ViewMatrix;
+		FMatrix InvViewMatrix;
 		FMatrix ViewProjectionMatrix;
 	};
 
@@ -61,6 +64,27 @@ public:
 		return ViewMatrices.GetViewProjectionMatrix().TransformFVector4(FVector4(WorldPoint, 1));
 	}
 
+	bool ScreenToPixel(const FVector4& ScreenPoint, FVector2D& OutPixelLocation) const
+	{
+		// implementation copied from FSceneView::ScreenToPixel() 
+
+		if (ScreenPoint.W != 0.0f)
+		{
+			//Reverse the W in the case it is negative, this allow to manipulate a manipulator in the same direction when the camera is really close to the manipulator.
+			float InvW = (ScreenPoint.W > 0.0f ? 1.0f : -1.0f) / ScreenPoint.W;
+			float Y = (GProjectionSignY > 0.0f) ? ScreenPoint.Y : 1.0f - ScreenPoint.Y;
+			OutPixelLocation = FVector2D(
+				UnscaledViewRect.Min.X + (0.5f + ScreenPoint.X * 0.5f * InvW) * UnscaledViewRect.Width(),
+				UnscaledViewRect.Min.Y + (0.5f - Y * 0.5f * InvW) * UnscaledViewRect.Height()
+			);
+			return true;
+		}
+		else
+		{
+			return false;
+		}
+	}
+
 	FMatrices ViewMatrices;
 	FIntRect UnscaledViewRect;
 	FVector	ViewLocation;
