Gizmos: modified UGizmoObjectModifyStateTarget so it can be constructed with either a gizmo manager or a tools manager.

diff --git a/Engine/Source/Runtime/InteractiveToolsFramework/Public/BaseGizmos/StateTargets.h b/Engine/Source/Runtime/InteractiveToolsFramework/Public/BaseGizmos/StateTargets.h
--- a/Engine/Source/Runtime/InteractiveToolsFramework/Public/BaseGizmos/StateTargets.h
+++ b/Engine/Source/Runtime/InteractiveToolsFramework/Public/BaseGizmos/StateTargets.h
@@ -89,9 +89,9 @@ class INTERACTIVETOOLSFRAMEWORK_API UGizmoObjectModifyStateTarget : public UObje
 public:
 	virtual void BeginUpdate()
 	{
-		if (GizmoManager.IsValid())
+		if (TransactionManager)
 		{
-			GizmoManager->BeginUndoTransaction(TransactionDescription);
+			TransactionManager->BeginUndoTransaction(TransactionDescription);
 		}
 		if (ModifyObject.IsValid())
 		{
@@ -101,9 +101,9 @@ public:
 
 	virtual void EndUpdate()
 	{
-		if (GizmoManager.IsValid())
+		if (TransactionManager)
 		{
-			GizmoManager->EndUndoTransaction();
+			TransactionManager->EndUndoTransaction();
 		}
 	}
 
@@ -119,10 +119,10 @@ public:
 	FText TransactionDescription;
 
 	/**
-	 * Pointer to the GizmoManager that is used to open/close the transaction
+	 * Pointer to the GizmoManager or ToolManager that is used to open/close the transaction
 	 */
-	TWeakObjectPtr<UInteractiveGizmoManager> GizmoManager;
-
+	UPROPERTY()
+	TScriptInterface<IToolContextTransactionProvider> TransactionManager;
 
 public:
 	/**
@@ -134,13 +134,17 @@ public:
 	static UGizmoObjectModifyStateTarget* Construct(
 		UObject* ModifyObjectIn,
 		const FText& DescriptionIn,
-		UInteractiveGizmoManager* GizmoManagerIn,
+		IToolContextTransactionProvider* TransactionManagerIn,
 		UObject* Outer = (UObject*)GetTransientPackage())
 	{
 		UGizmoObjectModifyStateTarget* NewTarget = NewObject<UGizmoObjectModifyStateTarget>(Outer);
 		NewTarget->ModifyObject = ModifyObjectIn;
 		NewTarget->TransactionDescription = DescriptionIn;
-		NewTarget->GizmoManager = GizmoManagerIn;
+
+		// have to explicitly configure this because we only have IToolContextTransactionProvider pointer
+		NewTarget->TransactionManager.SetInterface(TransactionManagerIn);
+		NewTarget->TransactionManager.SetObject(CastChecked<UObject>(TransactionManagerIn));
+
 		return NewTarget;
 	}
 };
