Gizmo Element Library: remove unused methods.

diff --git a/Engine/Source/Runtime/InteractiveToolsFramework/Private/BaseGizmos/GizmoElementArrow.cpp b/Engine/Source/Runtime/InteractiveToolsFramework/Private/BaseGizmos/GizmoElementArrow.cpp
--- a/Engine/Source/Runtime/InteractiveToolsFramework/Private/BaseGizmos/GizmoElementArrow.cpp
+++ b/Engine/Source/Runtime/InteractiveToolsFramework/Private/BaseGizmos/GizmoElementArrow.cpp
@@ -88,12 +88,6 @@ FInputRayHit UGizmoElementArrow::LineTrace(const UGizmoViewContext* ViewContext,
 	return Hit;
 }
 
-FBoxSphereBounds UGizmoElementArrow::CalcBounds(const FTransform& LocalToWorld) const
-{
-	// @todo - implement box-sphere bounds calculation
-	return FBoxSphereBounds();
-}
-
 void UGizmoElementArrow::SetBase(const FVector& InBase)
 {
 	if (Base != InBase)
diff --git a/Engine/Source/Runtime/InteractiveToolsFramework/Private/BaseGizmos/GizmoElementBox.cpp b/Engine/Source/Runtime/InteractiveToolsFramework/Private/BaseGizmos/GizmoElementBox.cpp
--- a/Engine/Source/Runtime/InteractiveToolsFramework/Private/BaseGizmos/GizmoElementBox.cpp
+++ b/Engine/Source/Runtime/InteractiveToolsFramework/Private/BaseGizmos/GizmoElementBox.cpp
@@ -56,12 +56,6 @@ FInputRayHit UGizmoElementBox::LineTrace(const UGizmoViewContext* ViewContext, c
 	return FInputRayHit();
 }
 
-FBoxSphereBounds UGizmoElementBox::CalcBounds(const FTransform& LocalToWorld) const
-{
-	// @todo - implement box-sphere bounds calculation
-	return FBoxSphereBounds();
-}
-
 void UGizmoElementBox::SetCenter(const FVector& InCenter)
 {
 	Center = InCenter;
diff --git a/Engine/Source/Runtime/InteractiveToolsFramework/Private/BaseGizmos/GizmoElementCircle.cpp b/Engine/Source/Runtime/InteractiveToolsFramework/Private/BaseGizmos/GizmoElementCircle.cpp
--- a/Engine/Source/Runtime/InteractiveToolsFramework/Private/BaseGizmos/GizmoElementCircle.cpp
+++ b/Engine/Source/Runtime/InteractiveToolsFramework/Private/BaseGizmos/GizmoElementCircle.cpp
@@ -111,12 +111,6 @@ FInputRayHit UGizmoElementCircle::LineTrace(const UGizmoViewContext* ViewContext
 	return FInputRayHit();
 }
 
-FBoxSphereBounds UGizmoElementCircle::CalcBounds(const FTransform& LocalToWorld) const
-{
-	// @todo - implement box-sphere bounds calculation
-	return FBoxSphereBounds();
-}
-
 void UGizmoElementCircle::SetCenter(FVector InCenter)
 {
 	Center = InCenter;
diff --git a/Engine/Source/Runtime/InteractiveToolsFramework/Private/BaseGizmos/GizmoElementCone.cpp b/Engine/Source/Runtime/InteractiveToolsFramework/Private/BaseGizmos/GizmoElementCone.cpp
--- a/Engine/Source/Runtime/InteractiveToolsFramework/Private/BaseGizmos/GizmoElementCone.cpp
+++ b/Engine/Source/Runtime/InteractiveToolsFramework/Private/BaseGizmos/GizmoElementCone.cpp
@@ -65,12 +65,6 @@ FInputRayHit UGizmoElementCone::LineTrace(const UGizmoViewContext* ViewContext,
 	return FInputRayHit();
 }
 
-FBoxSphereBounds UGizmoElementCone::CalcBounds(const FTransform& LocalToWorld) const
-{
-	// @todo - implement box-sphere bounds calculation
-	return FBoxSphereBounds();
-}
-
 void UGizmoElementCone::SetOrigin(const FVector& InOrigin)
 {
 	Origin = InOrigin;
diff --git a/Engine/Source/Runtime/InteractiveToolsFramework/Private/BaseGizmos/GizmoElementCylinder.cpp b/Engine/Source/Runtime/InteractiveToolsFramework/Private/BaseGizmos/GizmoElementCylinder.cpp
--- a/Engine/Source/Runtime/InteractiveToolsFramework/Private/BaseGizmos/GizmoElementCylinder.cpp
+++ b/Engine/Source/Runtime/InteractiveToolsFramework/Private/BaseGizmos/GizmoElementCylinder.cpp
@@ -65,12 +65,6 @@ FInputRayHit UGizmoElementCylinder::LineTrace(const UGizmoViewContext* ViewConte
 	return FInputRayHit();
 }
 
-FBoxSphereBounds UGizmoElementCylinder::CalcBounds(const FTransform& LocalToWorld) const
-{
-	// @todo - implement box-sphere bounds calculation
-	return FBoxSphereBounds();
-}
-
 void UGizmoElementCylinder::SetBase(const FVector& InBase)
 {
 	Base = InBase;
diff --git a/Engine/Source/Runtime/InteractiveToolsFramework/Private/BaseGizmos/GizmoElementGroup.cpp b/Engine/Source/Runtime/InteractiveToolsFramework/Private/BaseGizmos/GizmoElementGroup.cpp
--- a/Engine/Source/Runtime/InteractiveToolsFramework/Private/BaseGizmos/GizmoElementGroup.cpp
+++ b/Engine/Source/Runtime/InteractiveToolsFramework/Private/BaseGizmos/GizmoElementGroup.cpp
@@ -66,24 +66,6 @@ FInputRayHit UGizmoElementGroup::LineTrace(const UGizmoViewContext* ViewContext,
 	return Hit;
 }
 
-FBoxSphereBounds UGizmoElementGroup::CalcBounds(const FTransform& LocalToWorld) const
-{
-	/*  @todo - accumulate box sphere bounds for all elements within the group
-
-	if (bEnabled)
-	{
-		for (UGizmoElementBase* Element : Elements)
-		{
-			if (Element)
-			{
-				CachedBoxSphereBounds = CachedBoxSphereBounds + Element->CalcBounds(LocalToWorld);
-			}
-		}
-	}
-	*/
-	return FBoxSphereBounds();
-}
-
 void UGizmoElementGroup::Add(UGizmoElementBase* InElement)
 {
 	if (!Elements.Contains(InElement))
diff --git a/Engine/Source/Runtime/InteractiveToolsFramework/Private/BaseGizmos/GizmoElementRectangle.cpp b/Engine/Source/Runtime/InteractiveToolsFramework/Private/BaseGizmos/GizmoElementRectangle.cpp
--- a/Engine/Source/Runtime/InteractiveToolsFramework/Private/BaseGizmos/GizmoElementRectangle.cpp
+++ b/Engine/Source/Runtime/InteractiveToolsFramework/Private/BaseGizmos/GizmoElementRectangle.cpp
@@ -122,12 +122,6 @@ FInputRayHit UGizmoElementRectangle::LineTrace(const UGizmoViewContext* ViewCont
 	return FInputRayHit();
 }
 
-FBoxSphereBounds UGizmoElementRectangle::CalcBounds(const FTransform& LocalToWorld) const
-{
-	// @todo - implement box-sphere bounds calculation
-	return FBoxSphereBounds();
-}
-
 void UGizmoElementRectangle::SetCenter(FVector InCenter)
 {
 	Center = InCenter;
diff --git a/Engine/Source/Runtime/InteractiveToolsFramework/Private/BaseGizmos/GizmoElementTorus.cpp b/Engine/Source/Runtime/InteractiveToolsFramework/Private/BaseGizmos/GizmoElementTorus.cpp
--- a/Engine/Source/Runtime/InteractiveToolsFramework/Private/BaseGizmos/GizmoElementTorus.cpp
+++ b/Engine/Source/Runtime/InteractiveToolsFramework/Private/BaseGizmos/GizmoElementTorus.cpp
@@ -178,12 +178,6 @@ FInputRayHit UGizmoElementTorus::LineTrace(const UGizmoViewContext* ViewContext,
 	return FInputRayHit();
 }
 
-FBoxSphereBounds UGizmoElementTorus::CalcBounds(const FTransform& LocalToWorld) const
-{
-	// @todo - implement box-sphere bounds calculation
-	return FBoxSphereBounds();
-}
-
 void UGizmoElementTorus::SetCenter(const FVector& InCenter)
 {
 	Center = InCenter;
diff --git a/Engine/Source/Runtime/InteractiveToolsFramework/Public/BaseGizmos/GizmoElementArrow.h b/Engine/Source/Runtime/InteractiveToolsFramework/Public/BaseGizmos/GizmoElementArrow.h
--- a/Engine/Source/Runtime/InteractiveToolsFramework/Public/BaseGizmos/GizmoElementArrow.h
+++ b/Engine/Source/Runtime/InteractiveToolsFramework/Public/BaseGizmos/GizmoElementArrow.h
@@ -36,7 +36,6 @@ public:
 	//~ Begin UGizmoElementBase Interface.
 	virtual void Render(IToolsContextRenderAPI* RenderAPI, const FRenderTraversalState& RenderState) override;
 	virtual FInputRayHit LineTrace(const UGizmoViewContext* ViewContext, const FLineTraceTraversalState& LineTraceState, const FVector& RayOrigin, const FVector& RayDirection) override;
-	virtual FBoxSphereBounds CalcBounds(const FTransform& LocalToWorld) const override;
 	//~ End UGizmoElementBase Interface.
 
 	// Location of base of arrow cylinder
diff --git a/Engine/Source/Runtime/InteractiveToolsFramework/Public/BaseGizmos/GizmoElementBase.h b/Engine/Source/Runtime/InteractiveToolsFramework/Public/BaseGizmos/GizmoElementBase.h
--- a/Engine/Source/Runtime/InteractiveToolsFramework/Public/BaseGizmos/GizmoElementBase.h
+++ b/Engine/Source/Runtime/InteractiveToolsFramework/Public/BaseGizmos/GizmoElementBase.h
@@ -119,9 +119,6 @@ public:
 	// Line trace enabled hittable element.
 	virtual FInputRayHit LineTrace(const UGizmoViewContext* ViewContext, const FLineTraceTraversalState& LineTraceState, const FVector& RayOrigin, const FVector& RayDirection) PURE_VIRTUAL(UGizmoElementBase::LineTrace, return FInputRayHit(););
 
-	// Calcute box sphere bounds for use when hit testing.
-	virtual FBoxSphereBounds CalcBounds(const FTransform& LocalToWorld) const PURE_VIRTUAL(UGizmoElementBase::CalcBounds, return FBoxSphereBounds(););
-
 	// Return whether the visible bit is enabled in element state.
 	virtual bool GetVisibleState() const;
 
@@ -298,7 +295,7 @@ public:
 
 protected:
 
-	// Render and LineTrace should only occur when bEnabled is true.
+	// Whether element is enabled. Render and LineTrace only occur when bEnabled is true.
 	UPROPERTY()
 	bool bEnabled = true;
 
diff --git a/Engine/Source/Runtime/InteractiveToolsFramework/Public/BaseGizmos/GizmoElementBox.h b/Engine/Source/Runtime/InteractiveToolsFramework/Public/BaseGizmos/GizmoElementBox.h
--- a/Engine/Source/Runtime/InteractiveToolsFramework/Public/BaseGizmos/GizmoElementBox.h
+++ b/Engine/Source/Runtime/InteractiveToolsFramework/Public/BaseGizmos/GizmoElementBox.h
@@ -23,7 +23,6 @@ public:
 	//~ Begin UGizmoElementBase Interface.
 	virtual void Render(IToolsContextRenderAPI* RenderAPI, const FRenderTraversalState& RenderState) override;
 	virtual FInputRayHit LineTrace(const UGizmoViewContext* ViewContext, const FLineTraceTraversalState& LineTraceState, const FVector& RayOrigin, const FVector& RayDirection) override;
-	virtual FBoxSphereBounds CalcBounds(const FTransform& LocalToWorld) const override;
 	//~ End UGizmoElementBase Interface.
 
 	// Location of center of box's base circle.
diff --git a/Engine/Source/Runtime/InteractiveToolsFramework/Public/BaseGizmos/GizmoElementCircle.h b/Engine/Source/Runtime/InteractiveToolsFramework/Public/BaseGizmos/GizmoElementCircle.h
--- a/Engine/Source/Runtime/InteractiveToolsFramework/Public/BaseGizmos/GizmoElementCircle.h
+++ b/Engine/Source/Runtime/InteractiveToolsFramework/Public/BaseGizmos/GizmoElementCircle.h
@@ -21,7 +21,6 @@ public:
 	//~ Begin UGizmoElementBase Interface.
 	virtual void Render(IToolsContextRenderAPI* RenderAPI, const FRenderTraversalState& RenderState) override;
 	virtual FInputRayHit LineTrace(const UGizmoViewContext* ViewContext, const FLineTraceTraversalState& LineTraceState, const FVector& RayOrigin, const FVector& RayDirection) override;
-	virtual FBoxSphereBounds CalcBounds(const FTransform& LocalToWorld) const override;
 	//~ End UGizmoElementBase Interface.
 
 	// Circle center.
diff --git a/Engine/Source/Runtime/InteractiveToolsFramework/Public/BaseGizmos/GizmoElementCone.h b/Engine/Source/Runtime/InteractiveToolsFramework/Public/BaseGizmos/GizmoElementCone.h
--- a/Engine/Source/Runtime/InteractiveToolsFramework/Public/BaseGizmos/GizmoElementCone.h
+++ b/Engine/Source/Runtime/InteractiveToolsFramework/Public/BaseGizmos/GizmoElementCone.h
@@ -23,7 +23,6 @@ public:
 	//~ Begin UGizmoElementBase Interface.
 	virtual void Render(IToolsContextRenderAPI* RenderAPI, const FRenderTraversalState& RenderState) override;
 	virtual FInputRayHit LineTrace(const UGizmoViewContext* ViewContext, const FLineTraceTraversalState& LineTraceState, const FVector& RayOrigin, const FVector& RayDirection) override;
-	virtual FBoxSphereBounds CalcBounds(const FTransform& LocalToWorld) const override;
 	//~ End UGizmoElementBase Interface.
 
 	// Cone tip location.
diff --git a/Engine/Source/Runtime/InteractiveToolsFramework/Public/BaseGizmos/GizmoElementCylinder.h b/Engine/Source/Runtime/InteractiveToolsFramework/Public/BaseGizmos/GizmoElementCylinder.h
--- a/Engine/Source/Runtime/InteractiveToolsFramework/Public/BaseGizmos/GizmoElementCylinder.h
+++ b/Engine/Source/Runtime/InteractiveToolsFramework/Public/BaseGizmos/GizmoElementCylinder.h
@@ -23,7 +23,6 @@ public:
 	//~ Begin UGizmoElementBase Interface.
 	virtual void Render(IToolsContextRenderAPI* RenderAPI, const FRenderTraversalState& RenderState) override;
 	virtual FInputRayHit LineTrace(const UGizmoViewContext* ViewContext, const FLineTraceTraversalState& LineTraceState, const FVector& RayOrigin, const FVector& RayDirection) override;
-	virtual FBoxSphereBounds CalcBounds(const FTransform& LocalToWorld) const override;
 	//~ End UGizmoElementBase Interface.
 
 	// Location of center of cylinder's base circle.
diff --git a/Engine/Source/Runtime/InteractiveToolsFramework/Public/BaseGizmos/GizmoElementGroup.h b/Engine/Source/Runtime/InteractiveToolsFramework/Public/BaseGizmos/GizmoElementGroup.h
--- a/Engine/Source/Runtime/InteractiveToolsFramework/Public/BaseGizmos/GizmoElementGroup.h
+++ b/Engine/Source/Runtime/InteractiveToolsFramework/Public/BaseGizmos/GizmoElementGroup.h
@@ -22,7 +22,6 @@ public:
 	//~ Begin UGizmoElementBase Interface.
 	virtual void Render(IToolsContextRenderAPI* RenderAPI, const FRenderTraversalState& RenderState) override;
 	virtual FInputRayHit LineTrace(const UGizmoViewContext* ViewContext, const FLineTraceTraversalState& LineTraceState, const FVector& RayOrigin, const FVector& RayDirection) override;
-	virtual FBoxSphereBounds CalcBounds(const FTransform& LocalToWorld) const override;
 	//~ End UGizmoElementBase Interface.
 
 	// Add object to group.
diff --git a/Engine/Source/Runtime/InteractiveToolsFramework/Public/BaseGizmos/GizmoElementRectangle.h b/Engine/Source/Runtime/InteractiveToolsFramework/Public/BaseGizmos/GizmoElementRectangle.h
--- a/Engine/Source/Runtime/InteractiveToolsFramework/Public/BaseGizmos/GizmoElementRectangle.h
+++ b/Engine/Source/Runtime/InteractiveToolsFramework/Public/BaseGizmos/GizmoElementRectangle.h
@@ -25,7 +25,6 @@ public:
 	//~ Begin UGizmoElementBase Interface.
 	virtual void Render(IToolsContextRenderAPI* RenderAPI, const FRenderTraversalState& RenderState) override;
 	virtual FInputRayHit LineTrace(const UGizmoViewContext* ViewContext, const FLineTraceTraversalState& LineTraceState, const FVector& RayOrigin, const FVector& RayDirection) override;
-	virtual FBoxSphereBounds CalcBounds(const FTransform& LocalToWorld) const override;
 	//~ End UGizmoElementBase Interface.
 
 	// Location of rectangle center
diff --git a/Engine/Source/Runtime/InteractiveToolsFramework/Public/BaseGizmos/GizmoElementTorus.h b/Engine/Source/Runtime/InteractiveToolsFramework/Public/BaseGizmos/GizmoElementTorus.h
--- a/Engine/Source/Runtime/InteractiveToolsFramework/Public/BaseGizmos/GizmoElementTorus.h
+++ b/Engine/Source/Runtime/InteractiveToolsFramework/Public/BaseGizmos/GizmoElementTorus.h
@@ -33,7 +33,6 @@ public:
 	// If the torus lies at a glancing angle, the ray-torus intersection is performed against cylinders approximating
 	// the shape of the torus.
 	virtual FInputRayHit LineTrace(const UGizmoViewContext* ViewContext, const FLineTraceTraversalState& LineTraceState, const FVector& RayOrigin, const FVector& RayDirection) override;
-	virtual FBoxSphereBounds CalcBounds(const FTransform& LocalToWorld) const override;
 	//~ End UGizmoElementBase Interface.
 
 	// Torus center.
