Gizmo element library: fix variable names.

diff --git a/Engine/Source/Runtime/InteractiveToolsFramework/Private/BaseGizmos/GizmoElementBase.cpp b/Engine/Source/Runtime/InteractiveToolsFramework/Private/BaseGizmos/GizmoElementBase.cpp
--- a/Engine/Source/Runtime/InteractiveToolsFramework/Private/BaseGizmos/GizmoElementBase.cpp
+++ b/Engine/Source/Runtime/InteractiveToolsFramework/Private/BaseGizmos/GizmoElementBase.cpp
@@ -75,11 +75,11 @@ bool UGizmoElementBase::GetViewDependentVisibility(const FVector& InViewLocation
 	bool bVisibleViewDependent;
 	if (ViewDependentType == EGizmoElementViewDependentType::Axis)
 	{
-		bVisibleViewDependent = FMath::Abs(FVector::DotProduct(ViewDependentAxis, ViewDir)) < DefaultViewAlignAxialMaxCosAngleTol;
+		bVisibleViewDependent = FMath::Abs(FVector::DotProduct(ViewDependentAxis, ViewDir)) < DefaultViewDependentAxialMaxCosAngleTol;
 	}
 	else // (ViewDependentType == EGizmoElementViewDependentType::Plane)
 	{
-		bVisibleViewDependent = FMath::Abs(FVector::DotProduct(ViewDependentAxis, ViewDir)) > DefaultViewAlignPlanarMinCosAngleTol;
+		bVisibleViewDependent = FMath::Abs(FVector::DotProduct(ViewDependentAxis, ViewDir)) > DefaultViewDependentPlanarMinCosAngleTol;
 	}
 
 	return bVisibleViewDependent;
@@ -133,7 +133,7 @@ bool UGizmoElementBase::GetViewAlignRot(const FVector& InViewLocation, const FVe
 	else if (ViewAlignType == EGizmoElementViewAlignType::Axial)
 	{
 		// if Axis and Dir are almost coincident, do not adjust the rotation
-		if ((FMath::Abs(FVector::DotProduct(ViewAlignAxis, -LocalViewDir))) >= DefaultViewAlignAxialMaxCosAngleTol)
+		if ((FMath::Abs(FVector::DotProduct(ViewAlignAxis, -LocalViewDir))) >= DefaultViewAlignMaxCosAngleTol)
 		{
 			return false;
 		}
diff --git a/Engine/Source/Runtime/InteractiveToolsFramework/Private/BaseGizmos/GizmoElementTorus.cpp b/Engine/Source/Runtime/InteractiveToolsFramework/Private/BaseGizmos/GizmoElementTorus.cpp
--- a/Engine/Source/Runtime/InteractiveToolsFramework/Private/BaseGizmos/GizmoElementTorus.cpp
+++ b/Engine/Source/Runtime/InteractiveToolsFramework/Private/BaseGizmos/GizmoElementTorus.cpp
@@ -59,7 +59,7 @@ FInputRayHit UGizmoElementTorus::LineTrace(const UGizmoViewContext* ViewContext,
 		// angle between the vector from torus center to ring center and the vector from torus center to ring edge.
 		double MinCosAngle = 
 			FMath::Max(InnerRadius * 1.5 / FMath::Sqrt(OuterRadius * OuterRadius + InnerRadius * InnerRadius), 
-				DefaultViewAlignPlanarMinCosAngleTol);
+				DefaultViewDependentPlanarMinCosAngleTol);
 
 		bool bAtGlancingAngle = FMath::Abs(FVector::DotProduct(WorldNormal, ViewToTorusCenter)) <= MinCosAngle;
 
diff --git a/Engine/Source/Runtime/InteractiveToolsFramework/Public/BaseGizmos/GizmoElementBase.h b/Engine/Source/Runtime/InteractiveToolsFramework/Public/BaseGizmos/GizmoElementBase.h
--- a/Engine/Source/Runtime/InteractiveToolsFramework/Public/BaseGizmos/GizmoElementBase.h
+++ b/Engine/Source/Runtime/InteractiveToolsFramework/Public/BaseGizmos/GizmoElementBase.h
@@ -28,8 +28,8 @@ class INTERACTIVETOOLSFRAMEWORK_API UGizmoElementBase : public UObject
 public:
 
 	static constexpr float DefaultViewDependentAngleTol = 0.052f;			// ~3 degrees
-	static constexpr float DefaultViewAlignAxialMaxCosAngleTol = 0.998f;	// Cos(DefaultViewDependentAngleTol)
-	static constexpr float DefaultViewAlignPlanarMinCosAngleTol = 0.052f;	// Cos(HALF_PI - DefaultViewDependentAngleTol)
+	static constexpr float DefaultViewDependentAxialMaxCosAngleTol = 0.998f;	// Cos(DefaultViewDependentAngleTol)
+	static constexpr float DefaultViewDependentPlanarMinCosAngleTol = 0.052f;	// Cos(HALF_PI - DefaultViewDependentAngleTol)
 
 	static constexpr float DefaultViewAlignAngleTol = 0.052f;				// ~3 degrees
 	static constexpr float DefaultViewAlignMaxCosAngleTol = 0.998f;			// Cos(DefaultViewAlignAngleTol)
@@ -353,12 +353,12 @@ protected:
 	// Axial view alignment minimum cos angle tolerance, computed based on ViewDependentAngleTol. 
 	// When the cos of the angle between the view direction and the axis is less than this value, this object should not be culled.
 	UPROPERTY()
-	float ViewDependentAxialMaxCosAngleTol = DefaultViewAlignAxialMaxCosAngleTol;
+	float ViewDependentAxialMaxCosAngleTol = DefaultViewDependentAxialMaxCosAngleTol;
 
 	// Planar view alignment minimum cos angle tolerance, computed based on ViewDependentAngleTol. 
 	// When the cos of the angle between the view direction and the axis is greater than this value, this object should not be culled.
 	UPROPERTY()
-	float ViewDependentPlanarMinCosAngleTol = DefaultViewAlignPlanarMinCosAngleTol;
+	float ViewDependentPlanarMinCosAngleTol = DefaultViewDependentPlanarMinCosAngleTol;
 
 	// View align type: None, PointEye, or PointWorld.
 	// PointEye rotates this axis to align with the view up axis.
