ToolsFramework: add meta tags to AInternalToolFrameworkActor so that it and its subclasses cannot be placed via Place Actors, Blueprinted, etc

diff --git a/Engine/Source/Runtime/InteractiveToolsFramework/Public/InteractiveToolObjects.h b/Engine/Source/Runtime/InteractiveToolsFramework/Public/InteractiveToolObjects.h
--- a/Engine/Source/Runtime/InteractiveToolsFramework/Public/InteractiveToolObjects.h
+++ b/Engine/Source/Runtime/InteractiveToolsFramework/Public/InteractiveToolObjects.h
@@ -14,7 +14,7 @@
  * may need special-case handling, for example to prevent the user from
  * selecting and deleting them. 
  */
-UCLASS(Transient)
+UCLASS(Transient, NotPlaceable, Hidden, NotBlueprintable, NotBlueprintType)
 class INTERACTIVETOOLSFRAMEWORK_API AInternalToolFrameworkActor : public AActor
 {
 	GENERATED_BODY()
