Gizmo element library: updated IsVisible/IsHittable to test against the correct current interaction state.Added HoverVertexColor and InteractVertexColor to mesh render state (similar to LineColor in line render state).

diff --git a/Engine/Source/Runtime/InteractiveToolsFramework/Private/BaseGizmos/GizmoElementBase.cpp b/Engine/Source/Runtime/InteractiveToolsFramework/Private/BaseGizmos/GizmoElementBase.cpp
--- a/Engine/Source/Runtime/InteractiveToolsFramework/Private/BaseGizmos/GizmoElementBase.cpp
+++ b/Engine/Source/Runtime/InteractiveToolsFramework/Private/BaseGizmos/GizmoElementBase.cpp
@@ -185,11 +185,6 @@ FQuat UGizmoElementBase::GetAlignRotBetweenCoordSpaces(FVector SourceForward, FV
 	return Result;
 }
 
-bool UGizmoElementBase::GetEnabledForCurrentState(bool bIsPerspectiveProjection) const
-{
-	return (GetEnabledForInteractionState(ElementInteractionState) && GetEnabledForViewProjection(bIsPerspectiveProjection));
-}
-
 bool UGizmoElementBase::GetEnabledForInteractionState(const EGizmoElementInteractionState InElementInteractionState) const
 {
 	if (bEnabled)
@@ -227,17 +222,20 @@ bool UGizmoElementBase::GetEnabledForViewProjection(bool bIsPerspectiveProjectio
 	return false;
 }
 
-bool UGizmoElementBase::IsVisible(const FSceneView* View, const FTransform& InLocalToWorldTransform, const FVector& InLocalCenter) const
+bool UGizmoElementBase::IsVisible(const FSceneView* View, EGizmoElementInteractionState InCurrentInteractionState,
+	const FTransform& InLocalToWorldTransform, const FVector& InLocalCenter) const
 {
 	return (GetVisibleState() &&
-		GetEnabledForCurrentState(View->IsPerspectiveProjection()) &&
+		GetEnabledForInteractionState(InCurrentInteractionState) && 
+		GetEnabledForViewProjection(View->IsPerspectiveProjection()) &&
 		GetViewDependentVisibility(View, InLocalToWorldTransform, InLocalCenter));
 }
 
 bool UGizmoElementBase::IsHittable(const UGizmoViewContext* ViewContext, const FTransform& InLocalToWorldTransform, const FVector& InLocalCenter) const
 {
 	return (GetHittableState() &&
-		GetEnabledForCurrentState(ViewContext->IsPerspectiveProjection()) &&
+		GetEnabledForInteractionState(EGizmoElementInteractionState::None) &&
+		GetEnabledForViewProjection(ViewContext->IsPerspectiveProjection()) &&
 		(!GetVisibleState() || GetViewDependentVisibility(ViewContext, InLocalToWorldTransform, InLocalCenter)));
 }
 
@@ -264,7 +262,7 @@ bool UGizmoElementBase::UpdateRenderState(IToolsContextRenderAPI* RenderAPI, con
 
 	InOutRenderState.MeshRenderState.Update(MeshRenderAttributes);
 
-	if (IsVisible(View, InOutRenderState.LocalToWorldTransform, InLocalCenter))
+	if (IsVisible(View, InOutRenderState.InteractionState, InOutRenderState.LocalToWorldTransform, InLocalCenter))
 	{
 		bOutHasAlignRot = GetViewAlignRot(View, InOutRenderState.LocalToWorldTransform, InLocalCenter, OutAlignRot);
 		InOutRenderState.LocalToWorldTransform = FTransform(OutAlignRot, InLocalCenter) * InOutRenderState.LocalToWorldTransform;
@@ -281,17 +279,17 @@ bool UGizmoElementBase::UpdateLineTraceState(const UGizmoViewContext* ViewContex
 	return UpdateLineTraceState(ViewContext, InLocalCenter, InOutRenderState, bHasAlignRot, AlignQuat);
 }
 
-bool UGizmoElementBase::UpdateLineTraceState(const UGizmoViewContext* ViewContext, const FVector& InLocalCenter, FLineTraceTraversalState& InOutRenderState, bool& bOutHasAlignRot, FQuat& OutAlignRot)
+bool UGizmoElementBase::UpdateLineTraceState(const UGizmoViewContext* ViewContext, const FVector& InLocalCenter, FLineTraceTraversalState& InOutLineTraceState, bool& bOutHasAlignRot, FQuat& OutAlignRot)
 {
 	check(ViewContext);
 
 	OutAlignRot = FQuat::Identity;
 	bOutHasAlignRot = false;
 
-	if (IsHittable(ViewContext, InOutRenderState.LocalToWorldTransform, InLocalCenter))
+	if (IsHittable(ViewContext, InOutLineTraceState.LocalToWorldTransform, InLocalCenter))
 	{
-		bOutHasAlignRot = GetViewAlignRot(ViewContext, InOutRenderState.LocalToWorldTransform, InLocalCenter, OutAlignRot);
-		InOutRenderState.LocalToWorldTransform = FTransform(OutAlignRot, InLocalCenter) * InOutRenderState.LocalToWorldTransform;
+		bOutHasAlignRot = GetViewAlignRot(ViewContext, InOutLineTraceState.LocalToWorldTransform, InLocalCenter, OutAlignRot);
+		InOutLineTraceState.LocalToWorldTransform = FTransform(OutAlignRot, InLocalCenter) * InOutLineTraceState.LocalToWorldTransform;
 		return true;
 	}
 
@@ -303,11 +301,21 @@ bool UGizmoElementBase::GetVisibleState() const
 	return (static_cast<uint8>(ElementState) & static_cast<uint8>(EGizmoElementState::Visible));
 }
 
+void UGizmoElementBase::SetVisibleState(bool bVisible)
+{
+	ElementState = (bVisible ? ElementState | EGizmoElementState::Visible : ElementState & (~EGizmoElementState::Visible));
+}
+
 bool UGizmoElementBase::GetHittableState() const
 {
 	return (static_cast<uint8>(ElementState) & static_cast<uint8>(EGizmoElementState::Hittable));
 }
 
+void UGizmoElementBase::SetHittableState(bool bHittable)
+{
+	ElementState = (bHittable ? ElementState | EGizmoElementState::Hittable : ElementState & (~EGizmoElementState::Hittable));
+}
+
 void UGizmoElementBase::SetEnabled(bool InEnabled)
 {
 	bEnabled = InEnabled;
@@ -405,7 +413,7 @@ bool UGizmoElementBase::UpdatePartHittableState(bool bHittable, uint32 InPartIde
 		return false;
 	}
 
-	ElementState = (bHittable ? ElementState | EGizmoElementState::Hittable : ElementState & (~EGizmoElementState::Hittable));
+	SetHittableState(bHittable);
 	return true;
 }
 
@@ -421,7 +429,7 @@ bool UGizmoElementBase::UpdatePartVisibleState(bool bVisible, uint32 InPartIdent
 		return false;
 	}
 
-	ElementState = (bVisible ? ElementState | EGizmoElementState::Visible : ElementState & (~EGizmoElementState::Visible));
+	SetVisibleState(bVisible);
 	return true;
 }
 
@@ -547,7 +555,7 @@ const UMaterialInterface* UGizmoElementBase::GetMaterial() const
 	return MeshRenderAttributes.Material.GetMaterial();
 }
 
-bool UGizmoElementBase::GetMaterialOverridesChildState() const
+bool UGizmoElementBase::DoesMaterialOverrideChildState() const
 {
 	return MeshRenderAttributes.Material.bOverridesChildState;
 }
@@ -567,7 +575,7 @@ const UMaterialInterface* UGizmoElementBase::GetHoverMaterial() const
 	return MeshRenderAttributes.HoverMaterial.GetMaterial();
 }
 
-bool UGizmoElementBase::GetHoverMaterialOverridesChildState() const
+bool UGizmoElementBase::DoesHoverMaterialOverrideChildState() const
 {
 	return MeshRenderAttributes.HoverMaterial.bOverridesChildState;
 }
@@ -587,7 +595,7 @@ const UMaterialInterface* UGizmoElementBase::GetInteractMaterial() const
 	return MeshRenderAttributes.InteractMaterial.GetMaterial();
 }
 
-bool UGizmoElementBase::GetInteractMaterialOverridesChildState() const
+bool UGizmoElementBase::DoesInteractMaterialOverrideChildState() const
 {
 	return MeshRenderAttributes.InteractMaterial.bOverridesChildState;
 }
@@ -611,7 +619,7 @@ bool UGizmoElementBase::HasVertexColor() const
 	return MeshRenderAttributes.VertexColor.bHasValue;
 }
 
-bool UGizmoElementBase::GetVertexColorOverridesChildState() const
+bool UGizmoElementBase::DoesVertexColorOverrideChildState() const
 {
 	return MeshRenderAttributes.VertexColor.bOverridesChildState;
 }
@@ -621,3 +629,51 @@ void UGizmoElementBase::ClearVertexColor()
 	MeshRenderAttributes.VertexColor.Reset();
 }
 
+void UGizmoElementBase::SetHoverVertexColor(FLinearColor InVertexColor, bool InOverridesChildState)
+{
+	MeshRenderAttributes.HoverVertexColor.SetColor(InVertexColor, InOverridesChildState);
+}
+
+FLinearColor UGizmoElementBase::GetHoverVertexColor() const
+{
+	return MeshRenderAttributes.HoverVertexColor.GetColor();
+}
+bool UGizmoElementBase::HasHoverVertexColor() const
+{
+	return MeshRenderAttributes.HoverVertexColor.bHasValue;
+}
+
+bool UGizmoElementBase::DoesHoverVertexColorOverrideChildState() const
+{
+	return MeshRenderAttributes.HoverVertexColor.bOverridesChildState;
+}
+
+void UGizmoElementBase::ClearHoverVertexColor()
+{
+	MeshRenderAttributes.HoverVertexColor.Reset();
+}
+
+void UGizmoElementBase::SetInteractVertexColor(FLinearColor InInteractVertexColor, bool InOverridesChildState)
+{
+	MeshRenderAttributes.InteractVertexColor.SetColor(InInteractVertexColor, InOverridesChildState);
+}
+
+FLinearColor UGizmoElementBase::GetInteractVertexColor() const
+{
+	return MeshRenderAttributes.InteractVertexColor.GetColor();
+}
+bool UGizmoElementBase::HasInteractVertexColor() const
+{
+	return MeshRenderAttributes.InteractVertexColor.bHasValue;
+}
+
+bool UGizmoElementBase::DoesInteractVertexColorOverrideChildState() const
+{
+	return MeshRenderAttributes.InteractVertexColor.bOverridesChildState;
+}
+
+void UGizmoElementBase::ClearInteractVertexColor()
+{
+	MeshRenderAttributes.InteractVertexColor.Reset();
+}
+
diff --git a/Engine/Source/Runtime/InteractiveToolsFramework/Private/BaseGizmos/GizmoElementCircle.cpp b/Engine/Source/Runtime/InteractiveToolsFramework/Private/BaseGizmos/GizmoElementCircle.cpp
--- a/Engine/Source/Runtime/InteractiveToolsFramework/Private/BaseGizmos/GizmoElementCircle.cpp
+++ b/Engine/Source/Runtime/InteractiveToolsFramework/Private/BaseGizmos/GizmoElementCircle.cpp
@@ -29,7 +29,7 @@ void UGizmoElementCircle::Render(IToolsContextRenderAPI* RenderAPI, const FRende
 		{
 			if (const UMaterialInterface* UseMaterial = CurrentRenderState.GetCurrentMaterial())
 			{
-				FColor VertexColor = CurrentRenderState.GetVertexColor().ToFColor(false);
+				FColor VertexColor = CurrentRenderState.GetCurrentVertexColor().ToFColor(false);
 				DrawDisc(PDI, WorldCenter, WorldUpAxis, WorldSideAxis, VertexColor, WorldRadius, NumSides, UseMaterial->GetRenderProxy(), SDPG_Foreground);
 			}
 		}
diff --git a/Engine/Source/Runtime/InteractiveToolsFramework/Private/BaseGizmos/GizmoElementLineBase.cpp b/Engine/Source/Runtime/InteractiveToolsFramework/Private/BaseGizmos/GizmoElementLineBase.cpp
--- a/Engine/Source/Runtime/InteractiveToolsFramework/Private/BaseGizmos/GizmoElementLineBase.cpp
+++ b/Engine/Source/Runtime/InteractiveToolsFramework/Private/BaseGizmos/GizmoElementLineBase.cpp
@@ -99,7 +99,7 @@ bool UGizmoElementLineBase::HasLineColor() const
 	return LineRenderAttributes.LineColor.bHasValue;
 }
 
-bool UGizmoElementLineBase::GetLineColorOverridesChildState() const
+bool UGizmoElementLineBase::DoesLineColorOverrideChildState() const
 {
 	return LineRenderAttributes.LineColor.bOverridesChildState;
 }
@@ -123,7 +123,7 @@ bool UGizmoElementLineBase::HasHoverLineColor() const
 	return LineRenderAttributes.HoverLineColor.bHasValue;
 }
 
-bool UGizmoElementLineBase::GetHoverLineColorOverridesChildState() const
+bool UGizmoElementLineBase::DoesHoverLineColorOverrideChildState() const
 {
 	return LineRenderAttributes.HoverLineColor.bOverridesChildState;
 }
@@ -147,7 +147,7 @@ bool UGizmoElementLineBase::HasInteractLineColor() const
 	return LineRenderAttributes.InteractLineColor.bHasValue;
 }
 
-bool UGizmoElementLineBase::GetInteractLineColorOverridesChildState() const
+bool UGizmoElementLineBase::DoesInteractLineColorOverrideChildState() const
 {
 	return LineRenderAttributes.InteractLineColor.bOverridesChildState;
 }
diff --git a/Engine/Source/Runtime/InteractiveToolsFramework/Private/BaseGizmos/GizmoElementRectangle.cpp b/Engine/Source/Runtime/InteractiveToolsFramework/Private/BaseGizmos/GizmoElementRectangle.cpp
--- a/Engine/Source/Runtime/InteractiveToolsFramework/Private/BaseGizmos/GizmoElementRectangle.cpp
+++ b/Engine/Source/Runtime/InteractiveToolsFramework/Private/BaseGizmos/GizmoElementRectangle.cpp
@@ -27,7 +27,7 @@ void UGizmoElementRectangle::Render(IToolsContextRenderAPI* RenderAPI, const FRe
 		{
 			if (const UMaterialInterface* UseMaterial = CurrentRenderState.GetCurrentMaterial())
 			{
-				FColor VertexColor = CurrentRenderState.GetVertexColor().ToFColor(false);
+				FColor VertexColor = CurrentRenderState.GetCurrentVertexColor().ToFColor(false);
 				DrawRectangleMesh(PDI, WorldCenter, WorldUpAxis, WorldSideAxis, VertexColor, WorldWidth, WorldHeight, UseMaterial->GetRenderProxy(), SDPG_Foreground);
 			}
 		}
diff --git a/Engine/Source/Runtime/InteractiveToolsFramework/Private/BaseGizmos/GizmoElementRenderState.cpp b/Engine/Source/Runtime/InteractiveToolsFramework/Private/BaseGizmos/GizmoElementRenderState.cpp
--- a/Engine/Source/Runtime/InteractiveToolsFramework/Private/BaseGizmos/GizmoElementRenderState.cpp
+++ b/Engine/Source/Runtime/InteractiveToolsFramework/Private/BaseGizmos/GizmoElementRenderState.cpp
@@ -80,8 +80,16 @@ const UMaterialInterface* FGizmoElementMeshRenderStateAttributes::GetMaterial(EG
 	return Material.GetMaterial();
 }
 
-FLinearColor FGizmoElementMeshRenderStateAttributes::GetVertexColor()
+FLinearColor FGizmoElementMeshRenderStateAttributes::GetVertexColor(EGizmoElementInteractionState InteractionState)
 {
+	if (InteractionState == EGizmoElementInteractionState::Hovering)
+	{
+		return HoverVertexColor.GetColor();
+	}
+	else if (InteractionState == EGizmoElementInteractionState::Interacting)
+	{
+		return InteractVertexColor.GetColor();
+	}
 	return VertexColor.GetColor();
 }
 
diff --git a/Engine/Source/Runtime/InteractiveToolsFramework/Public/BaseGizmos/GizmoElementBase.h b/Engine/Source/Runtime/InteractiveToolsFramework/Public/BaseGizmos/GizmoElementBase.h
--- a/Engine/Source/Runtime/InteractiveToolsFramework/Public/BaseGizmos/GizmoElementBase.h
+++ b/Engine/Source/Runtime/InteractiveToolsFramework/Public/BaseGizmos/GizmoElementBase.h
@@ -76,10 +76,10 @@ public:
 			return MeshRenderState.GetMaterial(InteractionState);
 		}
 
-		// Returns the mesh vertex color.
-		FLinearColor GetVertexColor()
+		// Returns the mesh vertex color based on the current interaction state.
+		FLinearColor GetCurrentVertexColor()
 		{
-			return MeshRenderState.GetVertexColor();
+			return MeshRenderState.GetVertexColor(InteractionState);
 		}
 
 		// Returns the line color based on the current interaction state.
@@ -119,10 +119,12 @@ public:
 	// Line trace enabled hittable element.
 	virtual FInputRayHit LineTrace(const UGizmoViewContext* ViewContext, const FLineTraceTraversalState& LineTraceState, const FVector& RayOrigin, const FVector& RayDirection) PURE_VIRTUAL(UGizmoElementBase::LineTrace, return FInputRayHit(););
 
-	// Return whether the visible bit is enabled in element state.
+	// Set/get the visible bit in element state.
+	virtual void SetVisibleState(bool bVisible);
 	virtual bool GetVisibleState() const;
 
-	// Return whether the hittable bit is enabled in element state.
+	// Set/get the hittable bit in element state.
+	virtual void SetHittableState(bool bHittable);
 	virtual bool GetHittableState() const;
 
 	// Element enabled flag. Render and LineTrace only occur when bEnabled is true.
@@ -243,7 +245,7 @@ public:
 	virtual const UMaterialInterface* GetMaterial() const;
 
 	// Get mesh render state material attribute's override setting. 
-	virtual bool GetMaterialOverridesChildState() const;
+	virtual bool DoesMaterialOverrideChildState() const;
 
 	// Clear mesh render state material attribute. 
 	virtual void ClearMaterial();
@@ -257,7 +259,7 @@ public:
 	virtual const UMaterialInterface* GetHoverMaterial() const;
 
 	// Get mesh render state hover material attribute's override setting.
-	virtual bool GetHoverMaterialOverridesChildState() const;
+	virtual bool DoesHoverMaterialOverrideChildState() const;
 
 	// Clear mesh render state hover material attribute. 
 	virtual void ClearHoverMaterial();
@@ -271,7 +273,7 @@ public:
 	virtual const UMaterialInterface* GetInteractMaterial() const;
 
 	// Get mesh render state interact material attribute's override setting. 
-	virtual bool GetInteractMaterialOverridesChildState() const;
+	virtual bool DoesInteractMaterialOverrideChildState() const;
 
 	// Clear mesh render interact state material attribute. 
 	virtual void ClearInteractMaterial();
@@ -288,11 +290,45 @@ public:
 	virtual bool HasVertexColor() const;
 
 	// Get mesh render state vertex color attribute's override setting. 
-	virtual bool GetVertexColorOverridesChildState() const;
+	virtual bool DoesVertexColorOverrideChildState() const;
 
 	// Clear mesh render state vertex color attribute.
 	virtual void ClearVertexColor();
 
+	// Set mesh render state vertex color attribute. 
+	//  @param InVertexColor - vertex color to be set
+	//  @param InOverridesChildState - when true, this vertex color will override the material of all child elements.
+	virtual void SetHoverVertexColor(FLinearColor InVertexColor, bool InOverridesChildState = false);
+
+	// Get mesh render state vertex color attribute's value. 
+	virtual FLinearColor GetHoverVertexColor() const;
+
+	// Returns true, if mesh render state vertex color attribute has been set. 
+	virtual bool HasHoverVertexColor() const;
+
+	// Get mesh render state vertex color attribute's override setting. 
+	virtual bool DoesHoverVertexColorOverrideChildState() const;
+
+	// Clear mesh render state vertex color attribute.
+	virtual void ClearHoverVertexColor();
+
+	// Set mesh render state vertex color attribute. 
+	//  @param InVertexColor - vertex color to be set
+	//  @param InOverridesChildState - when true, this vertex color will override the material of all child elements.
+	virtual void SetInteractVertexColor(FLinearColor InVertexColor, bool InOverridesChildState = false);
+
+	// Get mesh render state vertex color attribute's value. 
+	virtual FLinearColor GetInteractVertexColor() const;
+
+	// Returns true, if mesh render state vertex color attribute has been set. 
+	virtual bool HasInteractVertexColor() const;
+
+	// Get mesh render state vertex color attribute's override setting. 
+	virtual bool DoesInteractVertexColorOverrideChildState() const;
+
+	// Clear mesh render state vertex color attribute.
+	virtual void ClearInteractVertexColor();
+
 protected:
 
 	// Whether element is enabled. Render and LineTrace only occur when bEnabled is true.
@@ -394,7 +430,8 @@ protected:
 protected:
 
 	// Return whether element is currently visible.
-	virtual bool IsVisible(const FSceneView* View, const FTransform& InLocalToWorldTransform, const FVector& InLocalCenter) const;
+	virtual bool IsVisible(const FSceneView* View, EGizmoElementInteractionState InCurrentInteractionState, 
+		const FTransform& InLocalToWorldTransform, const FVector& InLocalCenter) const;
 
 	// Return whether element is currently hittable.
 	virtual bool IsHittable(const UGizmoViewContext* ViewContext, const FTransform& InLocalToWorldTransform, const FVector& InLocalCenter) const;
@@ -408,9 +445,6 @@ protected:
 	// Returns whether object is visible based on view-dependent visibility settings.
 	virtual bool GetViewDependentVisibility(const FVector& InViewLocation, const FVector& InViewDirection, bool bInPerspectiveView, const FTransform& InLocalToWorldTransform, const FVector& InLocalCenter) const;
 
-	// Returns whether element is enabled based on element's current interaction state and input view projection type.
-	virtual bool GetEnabledForCurrentState(bool bIsPerspectiveProjection) const;
-
 	// Returns whether element is enabled for given interaction state.
 	virtual bool GetEnabledForInteractionState(EGizmoElementInteractionState InInteractionState) const;
 
diff --git a/Engine/Source/Runtime/InteractiveToolsFramework/Public/BaseGizmos/GizmoElementLineBase.h b/Engine/Source/Runtime/InteractiveToolsFramework/Public/BaseGizmos/GizmoElementLineBase.h
--- a/Engine/Source/Runtime/InteractiveToolsFramework/Public/BaseGizmos/GizmoElementLineBase.h
+++ b/Engine/Source/Runtime/InteractiveToolsFramework/Public/BaseGizmos/GizmoElementLineBase.h
@@ -59,7 +59,7 @@ public:
 	virtual bool HasLineColor() const;
 
 	// Get line render state line color attribute's override setting. 
-	virtual bool GetLineColorOverridesChildState() const;
+	virtual bool DoesLineColorOverrideChildState() const;
 
 	// Clear line render state line color attribute.
 	virtual void ClearLineColor();
@@ -76,7 +76,7 @@ public:
 	virtual bool HasHoverLineColor() const;
 
 	// Get line render state hover line color attribute's override setting.
-	virtual bool GetHoverLineColorOverridesChildState() const;
+	virtual bool DoesHoverLineColorOverrideChildState() const;
 
 	// Clear line render state hover line color attribute. 
 	virtual void ClearHoverLineColor();
@@ -93,7 +93,7 @@ public:
 	virtual bool HasInteractLineColor() const;
 
 	// Get line render state interact line color attribute's override setting. 
-	virtual bool GetInteractLineColorOverridesChildState() const;
+	virtual bool DoesInteractLineColorOverrideChildState() const;
 
 	// Clear line render state interact line color attribute.
 	virtual void ClearInteractLineColor();
diff --git a/Engine/Source/Runtime/InteractiveToolsFramework/Public/BaseGizmos/GizmoElementRenderState.h b/Engine/Source/Runtime/InteractiveToolsFramework/Public/BaseGizmos/GizmoElementRenderState.h
--- a/Engine/Source/Runtime/InteractiveToolsFramework/Public/BaseGizmos/GizmoElementRenderState.h
+++ b/Engine/Source/Runtime/InteractiveToolsFramework/Public/BaseGizmos/GizmoElementRenderState.h
@@ -106,11 +106,19 @@ struct INTERACTIVETOOLSFRAMEWORK_API FGizmoElementMeshRenderStateAttributes
 	UPROPERTY()
 	FGizmoElementColorAttribute VertexColor;
 
+	// Hover mesh vertex color
+	UPROPERTY()
+	FGizmoElementColorAttribute HoverVertexColor;
+
+	// Interact mesh vertex color
+	UPROPERTY()
+	FGizmoElementColorAttribute InteractVertexColor;
+
 	// Returns the material corresponding to the input interaction state
 	const UMaterialInterface* GetMaterial(EGizmoElementInteractionState InteractionState);
 
 	// Returns mesh vertex color
-	FLinearColor GetVertexColor();
+	FLinearColor GetVertexColor(EGizmoElementInteractionState InteractionState);
 
 	// Update this mesh render state based on a child mesh render state attribute.
 	void Update(FGizmoElementMeshRenderStateAttributes& InChildAttributes);
