Gizmo element library: added Arc and CircleBase element. Updated Torus and Circle to derive from CircleBase. Added partial options to CircleBase, including view-dependent partial type for instances of elements that should be drawn non-partial when aligned with viewer. This is needed for the new TRS gizmo rotation axes each of which is a semi-torus which becomes a full torus when aligned with the eye-to-torus direction.

diff --git a/Engine/Source/Editor/Experimental/EditorInteractiveToolsFramework/Private/EditorGizmos/TransformGizmo.cpp b/Engine/Source/Editor/Experimental/EditorInteractiveToolsFramework/Private/EditorGizmos/TransformGizmo.cpp
--- a/Engine/Source/Editor/Experimental/EditorInteractiveToolsFramework/Private/EditorGizmos/TransformGizmo.cpp
+++ b/Engine/Source/Editor/Experimental/EditorInteractiveToolsFramework/Private/EditorGizmos/TransformGizmo.cpp
@@ -399,19 +399,19 @@ void UTransformGizmo::EnableRotate(EAxisList::Type InAxisListToDraw)
 
 	if (bEnableX && RotateXAxisElement == nullptr)
 	{
-		RotateXAxisElement = MakeRotateAxis(ETransformGizmoPartIdentifier::RotateXAxis, XAxis, YAxis, ZAxis, AxisMaterialX, CurrentAxisMaterial);
+		RotateXAxisElement = MakeRotateAxis(ETransformGizmoPartIdentifier::RotateXAxis, YAxis, ZAxis, AxisMaterialX, CurrentAxisMaterial);
 		GizmoElementRoot->Add(RotateXAxisElement);
 	}
 
 	if (bEnableY && RotateYAxisElement == nullptr)
 	{
-		RotateYAxisElement = MakeRotateAxis(ETransformGizmoPartIdentifier::RotateYAxis, YAxis, ZAxis, XAxis, AxisMaterialY, CurrentAxisMaterial);
+		RotateYAxisElement = MakeRotateAxis(ETransformGizmoPartIdentifier::RotateYAxis, ZAxis, XAxis, AxisMaterialY, CurrentAxisMaterial);
 		GizmoElementRoot->Add(RotateYAxisElement);
 	}
 
 	if (bEnableZ && RotateZAxisElement == nullptr)
 	{
-		RotateZAxisElement = MakeRotateAxis(ETransformGizmoPartIdentifier::RotateZAxis, ZAxis, XAxis, YAxis, AxisMaterialZ, CurrentAxisMaterial);
+		RotateZAxisElement = MakeRotateAxis(ETransformGizmoPartIdentifier::RotateZAxis, XAxis, YAxis, AxisMaterialZ, CurrentAxisMaterial);
 		GizmoElementRoot->Add(RotateZAxisElement);
 	}
 
@@ -696,23 +696,25 @@ UGizmoElementRectangle* UTransformGizmo::MakeTranslateScreenSpaceHandle()
 	return RectangleElement;
 }
 
-UGizmoElementTorus* UTransformGizmo::MakeRotateAxis(ETransformGizmoPartIdentifier InPartId, const FVector& Normal, const FVector& TorusAxis0, const FVector& TorusAxis1,
+UGizmoElementTorus* UTransformGizmo::MakeRotateAxis(ETransformGizmoPartIdentifier InPartId, const FVector& TorusAxis0, const FVector& TorusAxis1,
 	UMaterialInterface* InMaterial, UMaterialInterface* InCurrentMaterial)
 {
 	UGizmoElementTorus* RotateAxisElement = NewObject<UGizmoElementTorus>();
 	RotateAxisElement->SetPartIdentifier(static_cast<uint32>(InPartId));
 	RotateAxisElement->SetCenter(FVector::ZeroVector);
-	RotateAxisElement->SetOuterRadius(UTransformGizmo::RotateAxisOuterRadius);
-	RotateAxisElement->SetOuterSegments(UTransformGizmo::RotateAxisOuterSegments);
+	RotateAxisElement->SetRadius(UTransformGizmo::RotateAxisOuterRadius);
+	RotateAxisElement->SetNumSegments(UTransformGizmo::RotateAxisNumSegments);
 	RotateAxisElement->SetInnerRadius(UTransformGizmo::RotateAxisInnerRadius);
-	RotateAxisElement->SetInnerSlices(UTransformGizmo::RotateAxisInnerSlices);
-	RotateAxisElement->SetNormal(Normal);
-	RotateAxisElement->SetBeginAxis(TorusAxis0);
-	RotateAxisElement->SetPartial(true);
-	RotateAxisElement->SetAngle(PI);
-	RotateAxisElement->SetViewDependentAxis(Normal);
+	RotateAxisElement->SetNumInnerSlices(UTransformGizmo::RotateAxisInnerSlices);
+	RotateAxisElement->SetAxis0(TorusAxis0);
+	RotateAxisElement->SetAxis1(TorusAxis1);
+	const FVector TorusNormal = RotateAxisElement->GetAxis0() ^ RotateAxisElement->GetAxis1();
+	RotateAxisElement->SetPartialType(EGizmoElementPartialType::PartialViewDependent);
+	RotateAxisElement->SetPartialStartAngle(0.0f);
+	RotateAxisElement->SetPartialEndAngle(UE_PI);
+	RotateAxisElement->SetViewDependentAxis(TorusNormal);
 	RotateAxisElement->SetViewAlignType(EGizmoElementViewAlignType::Axial);
-	RotateAxisElement->SetViewAlignAxis(Normal);
+	RotateAxisElement->SetViewAlignAxis(TorusNormal);
 	RotateAxisElement->SetViewAlignNormal(TorusAxis1);
 	RotateAxisElement->SetMaterial(InMaterial);
 	return RotateAxisElement;
@@ -724,7 +726,8 @@ UGizmoElementCircle* UTransformGizmo::MakeRotateCircleHandle(ETransformGizmoPart
 	CircleElement->SetPartIdentifier(static_cast<uint32>(InPartId));
 	CircleElement->SetCenter(FVector::ZeroVector);
 	CircleElement->SetRadius(InRadius);
-	CircleElement->SetNormal(-FVector::ForwardVector);
+	CircleElement->SetAxis0(FVector::UpVector);
+	CircleElement->SetAxis1(-FVector::RightVector);
 	CircleElement->SetLineColor(InColor);
 	CircleElement->SetViewAlignType(EGizmoElementViewAlignType::PointOnly);
 	CircleElement->SetViewAlignNormal(-FVector::ForwardVector);
diff --git a/Engine/Source/Editor/Experimental/EditorInteractiveToolsFramework/Public/EditorGizmos/TransformGizmo.h b/Engine/Source/Editor/Experimental/EditorInteractiveToolsFramework/Public/EditorGizmos/TransformGizmo.h
--- a/Engine/Source/Editor/Experimental/EditorInteractiveToolsFramework/Public/EditorGizmos/TransformGizmo.h
+++ b/Engine/Source/Editor/Experimental/EditorInteractiveToolsFramework/Public/EditorGizmos/TransformGizmo.h
@@ -101,7 +101,7 @@ public:
 	static constexpr float RotateArcballSphereRadius = 70.0f;
 	static constexpr float RotateAxisOuterRadius = 73.0f;
 	static constexpr float RotateAxisInnerRadius = 1.25f;
-	static constexpr int32 RotateAxisOuterSegments = 64;
+	static constexpr int32 RotateAxisNumSegments = 64;
 	static constexpr int32 RotateAxisInnerSlices = 8;
 	static constexpr float RotateOuterCircleRadius = 73.0f;
 	static constexpr float RotateScreenSpaceRadius = 83.0f;
@@ -389,7 +389,7 @@ protected:
 	virtual UGizmoElementArrow* MakeScaleAxis(ETransformGizmoPartIdentifier InPartId, const FVector& InAxisDir, const FVector& InSideDir, UMaterialInterface* InMaterial);
 
 	/** Construct rotate axis handle */
-	virtual UGizmoElementTorus* MakeRotateAxis(ETransformGizmoPartIdentifier InPartId, const FVector& Normal, const FVector& TorusAxis0, const FVector& TorusAxis1,
+	virtual UGizmoElementTorus* MakeRotateAxis(ETransformGizmoPartIdentifier InPartId, const FVector& TorusAxis0, const FVector& TorusAxis1,
 		UMaterialInterface* InMaterial, UMaterialInterface* InCurrentMaterial);
 
 	/** Construct uniform scale handle */
@@ -439,7 +439,7 @@ protected:
 	/** Handle click press for translate Z axis */
 	virtual void OnClickPressTranslateZAxis(const FInputDeviceRay& PressPos);
 
-	/** Handle click press for translate and scale X axes */
+	/** Handle click press for translate axes */
 	virtual void OnClickPressAxis(const FInputDeviceRay& PressPos);
 
 	/** Handle click drag for translate axes */
diff --git a/Engine/Source/Runtime/InteractiveToolsFramework/Private/BaseGizmos/GizmoElementArc.cpp b/Engine/Source/Runtime/InteractiveToolsFramework/Private/BaseGizmos/GizmoElementArc.cpp
new file mode 100644
--- /dev/null
+++ b/Engine/Source/Runtime/InteractiveToolsFramework/Private/BaseGizmos/GizmoElementArc.cpp
@@ -0,0 +1,229 @@
+// Copyright Epic Games, Inc. All Rights Reserved.
+ 
+#include "BaseGizmos/GizmoElementArc.h"
+#include "BaseGizmos/GizmoRenderingUtil.h"
+#include "BaseGizmos/GizmoMath.h"
+#include "DynamicMeshBuilder.h"
+#include "InputState.h"
+#include "Materials/MaterialInterface.h"
+#include "SceneManagement.h"
+#include "VectorUtil.h"
+
+
+static void DrawThickArc(FPrimitiveDrawInterface* PDI, const FVector& InCenter, const FVector& InAxis0, const FVector& InAxis1,
+	double InOuterRadius, double InInnerRadius, const double InStartAngle, const double InEndAngle, const int32 InNumSegments,
+	const FMaterialRenderProxy* MaterialRenderProxy, const FColor& InColor)
+{
+	// Implementation copied from FWidget::DrawThickArc. This should eventually be moved to PrimDrawingUtils.h/cpp.
+
+	if (InColor.A == 0)
+	{
+		return;
+	}
+
+	const int32 NumPoints = FMath::TruncToInt(InNumSegments * (InEndAngle - InStartAngle) / (UE_DOUBLE_PI / 2.0)) + 1;
+
+	FColor TriangleColor = InColor;
+	FColor RingColor = InColor;
+	RingColor.A = MAX_uint8;
+
+	FVector ZAxis = InAxis0 ^ InAxis1;
+	FVector LastWorldVertex;
+
+	FDynamicMeshBuilder MeshBuilder(PDI->View->GetFeatureLevel());
+
+	for (int32 RadiusIndex = 0; RadiusIndex < 2; ++RadiusIndex)
+	{
+		double Radius = (RadiusIndex == 0) ? InOuterRadius : InInnerRadius;
+		double TCRadius = Radius / (double)InOuterRadius;
+		//Compute vertices for base circle.
+		for (int32 VertexIndex = 0; VertexIndex <= NumPoints; VertexIndex++)
+		{
+			double Percent = VertexIndex / (double)NumPoints;
+			double Angle = FMath::Lerp(InStartAngle, InEndAngle, Percent);
+			double AngleDeg = FRotator::ClampAxis(Angle * 180.f / PI);
+
+			FVector VertexDir = InAxis0.RotateAngleAxis(AngleDeg, ZAxis);
+			VertexDir.Normalize();
+
+			double TCAngle = Percent * (PI / 2);
+			FVector2f TC(TCRadius * FMath::Cos(Angle), TCRadius * FMath::Sin(Angle));
+
+			// Keep the vertices in local space so that we don't lose precision when dealing with LWC
+			// The local-to-world transform is handled in the MeshBuilder.Draw() call at the end of this function
+			const FVector VertexPosition = VertexDir * Radius;
+			FVector Normal = VertexPosition;
+			Normal.Normalize();
+
+			FDynamicMeshVertex MeshVertex;
+			MeshVertex.Position = (FVector3f)VertexPosition;
+			MeshVertex.Color = TriangleColor;
+			MeshVertex.TextureCoordinate[0] = TC;
+
+			MeshVertex.SetTangents(
+				(FVector3f)-ZAxis,
+				FVector3f((-ZAxis) ^ Normal),
+				(FVector3f)Normal
+			);
+
+			MeshBuilder.AddVertex(MeshVertex); //Add bottom vertex
+
+			// Push out the arc line borders so they dont z-fight with the mesh arcs
+			// DrawLine needs vertices in world space, but this is fine because it takes FVectors and works with LWC well
+			FVector StartLinePos = LastWorldVertex;
+			FVector EndLinePos = VertexPosition + InCenter;
+			if (VertexIndex != 0)
+			{
+				PDI->DrawLine(StartLinePos, EndLinePos, RingColor, SDPG_Foreground);
+			}
+			LastWorldVertex = EndLinePos;
+		}
+	}
+
+	//Add top/bottom triangles, in the style of a fan.
+	int32 InnerVertexStartIndex = NumPoints + 1;
+	for (int32 VertexIndex = 0; VertexIndex < NumPoints; VertexIndex++)
+	{
+		MeshBuilder.AddTriangle(VertexIndex, VertexIndex + 1, InnerVertexStartIndex + VertexIndex);
+		MeshBuilder.AddTriangle(VertexIndex + 1, InnerVertexStartIndex + VertexIndex + 1, InnerVertexStartIndex + VertexIndex);
+	}
+
+	MeshBuilder.Draw(PDI, FTranslationMatrix(InCenter), MaterialRenderProxy, SDPG_Foreground, 0.f);
+}
+
+
+void UGizmoElementArc::Render(IToolsContextRenderAPI* RenderAPI, const FRenderTraversalState& RenderState)
+{
+	FRenderTraversalState CurrentRenderState(RenderState);
+	bool bVisibleViewDependent = UpdateRenderState(RenderAPI, Center, CurrentRenderState);
+
+	if (!bVisibleViewDependent)
+	{
+		return;
+	}
+
+	const UMaterialInterface* UseMaterial = CurrentRenderState.GetCurrentMaterial();
+	if (!UseMaterial)
+	{
+		return;
+	}
+	const FVector WorldCenter = CurrentRenderState.LocalToWorldTransform.TransformPosition(FVector::ZeroVector);
+	const FVector WorldAxis0 = CurrentRenderState.LocalToWorldTransform.TransformVectorNoScale(Axis0);
+	const FVector WorldAxis1 = CurrentRenderState.LocalToWorldTransform.TransformVectorNoScale(Axis1);
+	const FVector WorldNormal = WorldAxis0 ^ WorldAxis1;
+
+	bool bPartial = IsPartial(RenderAPI->GetSceneView(), WorldCenter, WorldNormal);
+	double Start, End;
+	if (bPartial)
+	{
+		if (PartialEndAngle - PartialStartAngle <= 0)
+		{
+			return;
+		}
+
+		Start = PartialStartAngle;
+		End = PartialEndAngle;
+	}
+	else
+	{
+		Start = 0.0;
+		End = UE_DOUBLE_TWO_PI;
+	}
+
+	const double WorldOuterRadius = Radius * CurrentRenderState.LocalToWorldTransform.GetScale3D().X;
+	const double WorldInnerRadius = InnerRadius * CurrentRenderState.LocalToWorldTransform.GetScale3D().X;
+	const FColor VertexColor = CurrentRenderState.GetCurrentVertexColor().ToFColor(false);
+
+	FPrimitiveDrawInterface* PDI = RenderAPI->GetPrimitiveDrawInterface();
+
+	DrawThickArc(PDI, WorldCenter, WorldAxis0, WorldAxis1, WorldOuterRadius, WorldInnerRadius,
+		Start, End, NumSegments, UseMaterial->GetRenderProxy(), VertexColor);
+}
+
+
+FInputRayHit UGizmoElementArc::LineTrace(const UGizmoViewContext* ViewContext, const FLineTraceTraversalState& LineTraceState, const FVector& RayOrigin, const FVector& RayDirection)
+{
+	FLineTraceTraversalState CurrentLineTraceState(LineTraceState);
+	bool bHittableViewDependent = UpdateLineTraceState(ViewContext, Center, CurrentLineTraceState);
+
+	if (bHittableViewDependent)
+	{
+		const double PixelHitThresholdAdjust = CurrentLineTraceState.PixelToWorldScale * PixelHitDistanceThreshold;
+		const FVector WorldCenter = CurrentLineTraceState.LocalToWorldTransform.TransformPosition(FVector::ZeroVector);
+		const FVector WorldAxis0 = CurrentLineTraceState.LocalToWorldTransform.TransformVectorNoScale(Axis0);
+		const FVector WorldAxis1 = CurrentLineTraceState.LocalToWorldTransform.TransformVectorNoScale(Axis1);
+		const FVector WorldNormal = WorldAxis0 ^ WorldAxis1;
+		double HitDepth = -1.0;
+
+		bool bPartial = IsPartial(ViewContext, WorldCenter, WorldNormal);
+		const double PartialAngle = PartialEndAngle - PartialStartAngle;
+
+		if (bPartial && PartialAngle <= 0.0)
+		{
+			return FInputRayHit();
+		}
+
+		// Intersect ray with plane in which arc lies.
+		FPlane Plane(WorldCenter, WorldNormal);
+		HitDepth = FMath::RayPlaneIntersectionParam(RayOrigin, RayDirection, Plane);
+		if (HitDepth < 0.0)
+		{
+			return FInputRayHit();
+		}
+
+		FVector HitPoint = RayOrigin + RayDirection * HitDepth;
+		FVector HitVec = HitPoint - WorldCenter;
+
+		// Determine whether hit point lies within the arc
+		const double WorldOuterRadius = Radius * CurrentLineTraceState.LocalToWorldTransform.GetScale3D().X;
+		const double WorldInnerRadius = InnerRadius * CurrentLineTraceState.LocalToWorldTransform.GetScale3D().X;
+		const double Distance = HitVec.Length();
+		const double HitBufferMax = WorldOuterRadius + PixelHitThresholdAdjust;
+		const double HitBufferMin = WorldInnerRadius - PixelHitThresholdAdjust;
+
+		if (Distance > HitBufferMax || Distance < HitBufferMin)
+		{
+			return FInputRayHit();
+		}
+
+		// Handle partial arc
+		if (bPartial)
+		{
+			// Compute projected angle
+			FVector WorldBeginAxis = WorldAxis0.RotateAngleAxis(PartialStartAngle, WorldNormal).GetSafeNormal();
+			FVector HitVecNormal = HitVec.GetSafeNormal();
+			double HitAngle = UE::Geometry::VectorUtil::PlaneAngleSignedR(WorldBeginAxis, HitVec, WorldNormal); 
+
+			if (HitAngle < 0)
+			{
+				HitAngle = UE_DOUBLE_TWO_PI + HitAngle;
+			}
+
+			if (HitAngle > PartialAngle)
+			{
+				return FInputRayHit();
+			}
+		}
+
+		if (HitDepth >= 0.0)
+		{
+			FInputRayHit RayHit(HitDepth);
+			RayHit.SetHitObject(this);
+			RayHit.HitIdentifier = PartIdentifier;
+			return RayHit;
+		}
+	}
+	
+	return FInputRayHit();
+}
+
+void UGizmoElementArc::SetInnerRadius(double InInnerRadius)
+{
+	InnerRadius = InInnerRadius;
+}
+
+double UGizmoElementArc::GetInnerRadius() const
+{
+	return InnerRadius;
+}
+
diff --git a/Engine/Source/Runtime/InteractiveToolsFramework/Private/BaseGizmos/GizmoElementCircle.cpp b/Engine/Source/Runtime/InteractiveToolsFramework/Private/BaseGizmos/GizmoElementCircle.cpp
--- a/Engine/Source/Runtime/InteractiveToolsFramework/Private/BaseGizmos/GizmoElementCircle.cpp
+++ b/Engine/Source/Runtime/InteractiveToolsFramework/Private/BaseGizmos/GizmoElementCircle.cpp
@@ -17,11 +17,8 @@ void UGizmoElementCircle::Render(IToolsContextRenderAPI* RenderAPI, const FRende
 	{
 		const float WorldRadius = Radius * CurrentRenderState.LocalToWorldTransform.GetScale3D().X;
 		const FVector WorldCenter = CurrentRenderState.LocalToWorldTransform.TransformPosition(FVector::ZeroVector);
-		FVector WorldUpAxis, WorldSideAxis;
-		const FVector WorldNormal = CurrentRenderState.LocalToWorldTransform.TransformVectorNoScale(Normal);
-		GizmoMath::MakeNormalPlaneBasis(WorldNormal, WorldUpAxis, WorldSideAxis);
-		WorldUpAxis.Normalize();
-		WorldSideAxis.Normalize();
+		const FVector WorldAxis0 = CurrentRenderState.LocalToWorldTransform.TransformVectorNoScale(Axis0);
+		const FVector WorldAxis1 = CurrentRenderState.LocalToWorldTransform.TransformVectorNoScale(Axis1);
 
 		FPrimitiveDrawInterface* PDI = RenderAPI->GetPrimitiveDrawInterface();
 
@@ -30,7 +27,7 @@ void UGizmoElementCircle::Render(IToolsContextRenderAPI* RenderAPI, const FRende
 			if (const UMaterialInterface* UseMaterial = CurrentRenderState.GetCurrentMaterial())
 			{
 				FColor VertexColor = CurrentRenderState.GetCurrentVertexColor().ToFColor(false);
-				DrawDisc(PDI, WorldCenter, WorldUpAxis, WorldSideAxis, VertexColor, WorldRadius, NumSides, UseMaterial->GetRenderProxy(), SDPG_Foreground);
+				DrawDisc(PDI, WorldCenter, WorldAxis0, WorldAxis1, VertexColor, WorldRadius, NumSegments, UseMaterial->GetRenderProxy(), SDPG_Foreground);
 			}
 		}
 
@@ -41,7 +38,7 @@ void UGizmoElementCircle::Render(IToolsContextRenderAPI* RenderAPI, const FRende
 			check(View);
 			float CurrentLineThickness = GetCurrentLineThickness(View->IsPerspectiveProjection(), View->FOV);
 
-			DrawCircle(PDI, WorldCenter, WorldUpAxis, WorldSideAxis, CurrentRenderState.GetCurrentLineColor(), WorldRadius, NumSides, SDPG_Foreground, CurrentLineThickness, 0.0, bScreenSpaceLine);
+			DrawCircle(PDI, WorldCenter, WorldAxis0, WorldAxis1, CurrentRenderState.GetCurrentLineColor(), WorldRadius, NumSegments, SDPG_Foreground, CurrentLineThickness, 0.0, bScreenSpaceLine);
 		}
 	}
 }
@@ -53,8 +50,10 @@ FInputRayHit UGizmoElementCircle::LineTrace(const UGizmoViewContext* ViewContext
 
 	if (bHittableViewDependent)
 	{
-		const FVector WorldNormal = CurrentLineTraceState.LocalToWorldTransform.TransformVectorNoScale(Normal);
-		const FVector WorldCenter = CurrentLineTraceState.LocalToWorldTransform.TransformPosition(FVector::ZeroVector);
+		const FVector WorldCenter = CurrentLineTraceState.LocalToWorldTransform.GetLocation();
+		const FVector WorldAxis0 = CurrentLineTraceState.LocalToWorldTransform.TransformVectorNoScale(Axis0);
+		const FVector WorldAxis1 = CurrentLineTraceState.LocalToWorldTransform.TransformVectorNoScale(Axis1);
+		const FVector WorldNormal = WorldAxis0 ^ WorldAxis1;
 		const double PixelHitThresholdAdjust = CurrentLineTraceState.PixelToWorldScale * PixelHitDistanceThreshold;
 		double WorldRadius = CurrentLineTraceState.LocalToWorldTransform.GetScale3D().X * Radius;
 
@@ -111,47 +110,6 @@ FInputRayHit UGizmoElementCircle::LineTrace(const UGizmoViewContext* ViewContext
 	return FInputRayHit();
 }
 
-void UGizmoElementCircle::SetCenter(FVector InCenter)
-{
-	Center = InCenter;
-}
-
-FVector UGizmoElementCircle::GetCenter() const
-{
-	return Center;
-}
-
-void UGizmoElementCircle::SetNormal(FVector InNormal)
-{
-	Normal = InNormal;
-	Normal.Normalize();
-}
-
-FVector UGizmoElementCircle::GetNormal() const
-{
-	return Normal;
-}
-
-void UGizmoElementCircle::SetRadius(float InRadius)
-{
-	Radius = InRadius;
-}
-
-float UGizmoElementCircle::GetRadius() const
-{
-	return Radius;
-}
-
-void UGizmoElementCircle::SetNumSides(int InNumSides)
-{
-	NumSides = InNumSides;
-}
-
-int UGizmoElementCircle::GetNumSides() const
-{
-	return NumSides;
-}
-
 void UGizmoElementCircle::SetDrawMesh(bool InDrawMesh)
 {
 	bDrawMesh = InDrawMesh;
diff --git a/Engine/Source/Runtime/InteractiveToolsFramework/Private/BaseGizmos/GizmoElementCircleBase.cpp b/Engine/Source/Runtime/InteractiveToolsFramework/Private/BaseGizmos/GizmoElementCircleBase.cpp
new file mode 100644
--- /dev/null
+++ b/Engine/Source/Runtime/InteractiveToolsFramework/Private/BaseGizmos/GizmoElementCircleBase.cpp
@@ -0,0 +1,131 @@
+// Copyright Epic Games, Inc. All Rights Reserved.
+ 
+#include "BaseGizmos/GizmoElementCircleBase.h"
+#include "BaseGizmos/GizmoRenderingUtil.h"
+#include "BaseGizmos/GizmoMath.h"
+#include "InputState.h"
+#include "Materials/MaterialInterface.h"
+#include "SceneManagement.h"
+#include "VectorUtil.h"
+
+bool UGizmoElementCircleBase::IsPartial(const FSceneView* View, const FVector& InWorldCenter, const FVector& InWorldNormal)
+{
+	return IsPartial(InWorldCenter, InWorldNormal, View->ViewLocation, View->GetViewDirection(),
+		View->IsPerspectiveProjection());
+}
+
+bool UGizmoElementCircleBase::IsPartial(const UGizmoViewContext* View, const FVector& InWorldCenter, const FVector& InWorldNormal)
+{
+	return IsPartial(InWorldCenter, InWorldNormal, View->ViewLocation, View->GetViewDirection(),
+		View->IsPerspectiveProjection());
+}
+
+bool UGizmoElementCircleBase::IsPartial(const FVector& InWorldCenter, const FVector& InWorldNormal,
+	const FVector& InViewLocation, const FVector& InViewDirection, bool bIsPerspectiveProjection)
+{
+	FVector ViewToCircleBaseCenter;
+	if (bIsPerspectiveProjection)
+	{
+		ViewToCircleBaseCenter = (InWorldCenter - InViewLocation).GetSafeNormal();
+	}
+	else
+	{
+		ViewToCircleBaseCenter = InViewDirection;
+	}
+
+	double DotP = FMath::Abs(FVector::DotProduct(InWorldNormal, ViewToCircleBaseCenter));
+	return (PartialType == EGizmoElementPartialType::Partial ||
+		(PartialType == EGizmoElementPartialType::PartialViewDependent && DotP <= PartialViewDependentMaxCosTol));
+}
+
+void UGizmoElementCircleBase::SetCenter(const FVector& InCenter)
+{
+	Center = InCenter;
+}
+
+FVector UGizmoElementCircleBase::GetCenter() const
+{
+	return Center;
+}
+
+void UGizmoElementCircleBase::SetAxis0(const FVector& InAxis0)
+{
+	Axis0 = InAxis0;
+	Axis0.Normalize();
+}
+
+FVector UGizmoElementCircleBase::GetAxis0() const
+{
+	return Axis1;
+}
+
+void UGizmoElementCircleBase::SetAxis1(const FVector& InAxis1)
+{
+	Axis1 = InAxis1;
+	Axis1.Normalize();
+}
+
+FVector UGizmoElementCircleBase::GetAxis1() const
+{
+	return Axis0;
+}
+
+void UGizmoElementCircleBase::SetRadius(double InRadius)
+{
+	Radius = InRadius;
+}
+
+double UGizmoElementCircleBase::GetRadius() const
+{
+	return Radius;
+}
+
+void UGizmoElementCircleBase::SetNumSegments(int32 InNumSegments)
+{
+	NumSegments = InNumSegments;
+}
+
+int32 UGizmoElementCircleBase::GetNumSegments() const
+{
+	return NumSegments;
+}
+
+void UGizmoElementCircleBase::SetPartialType(EGizmoElementPartialType InPartialType)
+{
+	PartialType = InPartialType;
+}
+
+EGizmoElementPartialType UGizmoElementCircleBase::GetPartialType() const
+{
+	return PartialType;
+}
+
+void UGizmoElementCircleBase::SetPartialStartAngle(double InPartialStartAngle)
+{
+	PartialStartAngle = InPartialStartAngle;
+}
+
+double UGizmoElementCircleBase::GetPartialStartAngle() const
+{
+	return PartialStartAngle;
+}
+
+void UGizmoElementCircleBase::SetPartialEndAngle(double InPartialEndAngle)
+{
+	PartialEndAngle = InPartialEndAngle;
+}
+
+double UGizmoElementCircleBase::GetPartialEndAngle() const
+{
+	return PartialEndAngle;
+}
+
+void UGizmoElementCircleBase::SetPartialViewDependentMaxCosTol(double InPartialViewDependentMaxCosTol)
+{
+	PartialViewDependentMaxCosTol = InPartialViewDependentMaxCosTol;
+}
+
+double UGizmoElementCircleBase::GetPartialViewDependentMaxCosTol() const
+{
+	return PartialViewDependentMaxCosTol;
+}
diff --git a/Engine/Source/Runtime/InteractiveToolsFramework/Private/BaseGizmos/GizmoElementTorus.cpp b/Engine/Source/Runtime/InteractiveToolsFramework/Private/BaseGizmos/GizmoElementTorus.cpp
--- a/Engine/Source/Runtime/InteractiveToolsFramework/Private/BaseGizmos/GizmoElementTorus.cpp
+++ b/Engine/Source/Runtime/InteractiveToolsFramework/Private/BaseGizmos/GizmoElementTorus.cpp
@@ -17,14 +17,28 @@ void UGizmoElementTorus::Render(IToolsContextRenderAPI* RenderAPI, const FRender
 	{
 		if (const UMaterialInterface* UseMaterial = CurrentRenderState.GetCurrentMaterial())
 		{
+			const FVector WorldCenter = CurrentRenderState.LocalToWorldTransform.GetLocation();
+			const FVector WorldAxis0 = CurrentRenderState.LocalToWorldTransform.TransformVectorNoScale(Axis0);
+			const FVector Normal = Axis0 ^ Axis1; 
+			const FVector WorldNormal = CurrentRenderState.LocalToWorldTransform.TransformVectorNoScale(Normal);
+
+			bool bPartial = IsPartial(RenderAPI->GetSceneView(), WorldCenter, WorldNormal);
+			FVector BeginAxis = bPartial ? Axis0.RotateAngleAxis(PartialStartAngle, Normal).GetSafeNormal() : Axis0;
+
+			const double PartialAngle = PartialEndAngle - PartialStartAngle;			
+			if (PartialAngle <= 0.0)
+			{
+				return;
+			}
+
 			FPrimitiveDrawInterface* PDI = RenderAPI->GetPrimitiveDrawInterface();
 
 			FVector TorusSideAxis = Normal ^ BeginAxis;
 			TorusSideAxis.Normalize();
 
 			DrawTorus(PDI, CurrentRenderState.LocalToWorldTransform.ToMatrixWithScale(), 
-				BeginAxis, TorusSideAxis, OuterRadius, InnerRadius, OuterSegments, InnerSlices,  
-				UseMaterial->GetRenderProxy(), SDPG_Foreground, bPartial, Angle, bEndCaps);
+				BeginAxis, TorusSideAxis, Radius, InnerRadius, NumSegments, NumInnerSlices,  
+				UseMaterial->GetRenderProxy(), SDPG_Foreground, bPartial, PartialAngle, bEndCaps);
 		}
 	}
 }
@@ -45,31 +59,39 @@ FInputRayHit UGizmoElementTorus::LineTrace(const UGizmoViewContext* ViewContext,
 	if (bHittableViewDependent)
 	{
 		const double PixelHitThresholdAdjust = CurrentLineTraceState.PixelToWorldScale * PixelHitDistanceThreshold;
-		const double WorldOuterRadius = OuterRadius * CurrentLineTraceState.LocalToWorldTransform.GetScale3D().X;
+		const double WorldOuterRadius = Radius * CurrentLineTraceState.LocalToWorldTransform.GetScale3D().X;
 		const double WorldInnerRadius = InnerRadius * CurrentLineTraceState.LocalToWorldTransform.GetScale3D().X;
-		const FVector WorldCenter = CurrentLineTraceState.LocalToWorldTransform.TransformPosition(FVector::ZeroVector);
-		const FVector WorldNormal = CurrentLineTraceState.LocalToWorldTransform.TransformVectorNoScale(Normal);
-		const FVector WorldBeginAxis = CurrentLineTraceState.LocalToWorldTransform.TransformVectorNoScale(BeginAxis);
+		const FVector WorldCenter = CurrentLineTraceState.LocalToWorldTransform.GetLocation();
+		const FVector WorldAxis0 = CurrentLineTraceState.LocalToWorldTransform.TransformVectorNoScale(Axis0);
+		const FVector WorldAxis1 = CurrentLineTraceState.LocalToWorldTransform.TransformVectorNoScale(Axis1);
+		const FVector WorldNormal = WorldAxis0 ^ WorldAxis1;
 		double HitDepth = -1.0;
 
-		const FVector ViewToTorusCenter = 
-			ViewContext->IsPerspectiveProjection() ? (WorldCenter - ViewContext->ViewLocation).GetSafeNormal() : ViewContext->GetViewDirection();
+		bool bPartial = IsPartial(ViewContext, WorldCenter, WorldNormal);
 
-		// Determine if the ray direction is at a glancing angle by using a minimum cos angle based on the
-		// angle between the vector from torus center to ring center and the vector from torus center to ring edge.
-		double MinCosAngle = 
-			FMath::Max(InnerRadius * 1.5 / FMath::Sqrt(OuterRadius * OuterRadius + InnerRadius * InnerRadius), 
-				DefaultViewDependentPlanarMinCosAngleTol);
+		FVector WorldBeginAxis;
+		const double PartialAngle = PartialEndAngle - PartialStartAngle;
+		if (bPartial)
+		{
+			if (PartialAngle <= 0)
+			{
+				return FInputRayHit();
+			}
 
-		bool bAtGlancingAngle = FMath::Abs(FVector::DotProduct(WorldNormal, ViewToTorusCenter)) <= MinCosAngle;
+			WorldBeginAxis = WorldAxis0.RotateAngleAxis(PartialStartAngle, WorldNormal).GetSafeNormal();
+		}
+		else
+		{
+			WorldBeginAxis = WorldAxis0;
+		}
 
-		if (bAtGlancingAngle)
+		if (IsGlancingAngle(ViewContext, WorldCenter, WorldNormal))
 		{
 			// If the torus lies at a glancing angle, the ray-torus intersection is performed against cylinders approximating
 			// the shape of the torus.
 			static constexpr int NumFullTorusCylinders = 16;
 			static constexpr double AngleDelta = UE_DOUBLE_TWO_PI / static_cast<double>(NumFullTorusCylinders);
-			const int NumCylinders = bPartial ? FMath::CeilToInt(NumFullTorusCylinders * Angle / UE_DOUBLE_TWO_PI) : NumFullTorusCylinders;
+			const int NumCylinders = bPartial ? FMath::CeilToInt(NumFullTorusCylinders * PartialAngle / UE_DOUBLE_TWO_PI) : NumFullTorusCylinders;
 
 			const FVector ViewDirection = ViewContext->GetViewDirection();
 			FVector VectorA = WorldBeginAxis; 
@@ -94,7 +116,7 @@ FInputRayHit UGizmoElementTorus::LineTrace(const UGizmoViewContext* ViewContext,
 
 				if (i == NumCylinders - 1)
 				{
-					CylinderHeight = bPartial ? FMath::Fmod(Angle, AngleDelta) * CylinderHeight : CylinderHeight;
+					CylinderHeight = bPartial ? FMath::Fmod(PartialAngle, AngleDelta) * CylinderHeight : CylinderHeight;
 				}
 
 				const FVector CylinderDirection = (VectorB - VectorA).GetSafeNormal();
@@ -145,6 +167,7 @@ FInputRayHit UGizmoElementTorus::LineTrace(const UGizmoViewContext* ViewContext,
 				return FInputRayHit();
 			}
 
+			// Handle partial torus
 			if (bPartial)
 			{
 				// Compute projected angle
@@ -153,10 +176,10 @@ FInputRayHit UGizmoElementTorus::LineTrace(const UGizmoViewContext* ViewContext,
 
 				if (HitAngle < 0)
 				{
-					HitAngle = UE_DOUBLE_TWO_PI - HitAngle;
+					HitAngle = UE_DOUBLE_TWO_PI + HitAngle;
 				}
 
-				if (HitAngle > Angle)
+				if (HitAngle > PartialAngle)
 				{
 					return FInputRayHit();
 				}
@@ -178,96 +201,47 @@ FInputRayHit UGizmoElementTorus::LineTrace(const UGizmoViewContext* ViewContext,
 	return FInputRayHit();
 }
 
-void UGizmoElementTorus::SetCenter(const FVector& InCenter)
-{
-	Center = InCenter;
-}
-
-FVector UGizmoElementTorus::GetCenter() const
-{
-	return Center;
-}
-
-void UGizmoElementTorus::SetNormal(const FVector& InNormal)
-{
-	Normal = InNormal;
-	Normal.Normalize();
-}
-
-FVector UGizmoElementTorus::GetNormal() const
-{
-	return Normal;
-}
-
-void UGizmoElementTorus::SetBeginAxis(const FVector& InBeginAxis)
-{
-	BeginAxis = InBeginAxis;
-	BeginAxis.Normalize();
-}
-
-FVector UGizmoElementTorus::GetBeginAxis() const
+bool UGizmoElementTorus::IsGlancingAngle(const UGizmoViewContext* View, const FVector& InWorldCenter, const FVector& InWorldNormal)
 {
-	return BeginAxis;
-}
+	check(View);
+	FVector ViewToCircleBaseCenter;
+	if (View->IsPerspectiveProjection())
+	{
+		ViewToCircleBaseCenter = (InWorldCenter - View->ViewLocation).GetSafeNormal();
+	}
+	else
+	{
+		ViewToCircleBaseCenter = View->GetViewDirection();
+	}
 
-void UGizmoElementTorus::SetOuterRadius(float InOuterRadius)
-{
-	OuterRadius = InOuterRadius;
-}
+	// Determine if the ray direction is at a glancing angle by using a minimum cos angle based on the
+	// angle between the vector from arc center to ring center and the vector from arc center to ring edge.
+	double MinCosAngle =
+		FMath::Max(InnerRadius * 1.5 / FMath::Sqrt(Radius * Radius + InnerRadius * InnerRadius),
+			DefaultViewDependentPlanarMinCosAngleTol);
 
-float UGizmoElementTorus::GetOuterRadius() const
-{
-	return OuterRadius;
+	double DotP = FMath::Abs(FVector::DotProduct(InWorldNormal, ViewToCircleBaseCenter));
+	return (DotP <= MinCosAngle);
 }
 
-void UGizmoElementTorus::SetInnerRadius(float InInnerRadius)
+void UGizmoElementTorus::SetInnerRadius(double InInnerRadius)
 {
 	InnerRadius = InInnerRadius;
 }
 
-float UGizmoElementTorus::GetInnerRadius() const
+double UGizmoElementTorus::GetInnerRadius() const
 {
 	return InnerRadius;
 }
 
-void UGizmoElementTorus::SetOuterSegments(int32 InOuterSegments)
-{
-	OuterSegments = InOuterSegments;
-}
-
-int32 UGizmoElementTorus::GetOuterSegments() const
-{
-	return OuterSegments;
-}
-
-void UGizmoElementTorus::SetInnerSlices(int32 InInnerSlices)
-{
-	InnerSlices = InInnerSlices;
-}
-
-int32 UGizmoElementTorus::GetInnerSlices() const
-{
-	return InnerSlices;
-}
-
-void UGizmoElementTorus::SetPartial(bool InPartial)
-{
-	bPartial = InPartial;
-}
-
-bool UGizmoElementTorus::GetPartial() const
-{
-	return bPartial;
-}
-
-void UGizmoElementTorus::SetAngle(float InAngle)
+void UGizmoElementTorus::SetNumInnerSlices(int32 InNumInnerSlices)
 {
-	Angle = InAngle;
+	NumInnerSlices = InNumInnerSlices;
 }
 
-float UGizmoElementTorus::GetAngle() const
+int32 UGizmoElementTorus::GetNumInnerSlices() const
 {
-	return Angle;
+	return NumInnerSlices;
 }
 
 void UGizmoElementTorus::SetEndCaps(bool InEndCaps)
diff --git a/Engine/Source/Runtime/InteractiveToolsFramework/Public/BaseGizmos/GizmoElementArc.h b/Engine/Source/Runtime/InteractiveToolsFramework/Public/BaseGizmos/GizmoElementArc.h
new file mode 100644
--- /dev/null
+++ b/Engine/Source/Runtime/InteractiveToolsFramework/Public/BaseGizmos/GizmoElementArc.h
@@ -0,0 +1,39 @@
+// Copyright Epic Games, Inc. All Rights Reserved.
+
+#pragma once
+
+#include "CoreMinimal.h"
+#include "BaseGizmos/GizmoElementCircleBase.h"
+#include "InputState.h"
+#include "UObject/ObjectMacros.h"
+#include "GizmoElementArc.generated.h"
+
+class FPrimitiveDrawInterface;
+class FMaterialRenderProxy;
+
+/**
+ * Simple object intended to be used as part of 3D Gizmos.
+ * Draws a thick arc based on parameters.
+ */
+UCLASS(Transient)
+class INTERACTIVETOOLSFRAMEWORK_API UGizmoElementArc : public UGizmoElementCircleBase
+{
+	GENERATED_BODY()
+
+public:
+
+	//~ Begin UGizmoElementBase Interface.
+	virtual void Render(IToolsContextRenderAPI* RenderAPI, const FRenderTraversalState& RenderState) override;
+	virtual FInputRayHit LineTrace(const UGizmoViewContext* ViewContext, const FLineTraceTraversalState& LineTraceState, const FVector& RayOrigin, const FVector& RayDirection) override;
+	//~ End UGizmoElementBase Interface.
+
+	// Inner circle radius.
+	virtual void SetInnerRadius(double InInnerRadius);
+	virtual double GetInnerRadius() const;
+
+protected:
+
+	// Arc inner radius.
+	UPROPERTY()
+	double InnerRadius = 5.0;
+};
\ No newline at end of file
diff --git a/Engine/Source/Runtime/InteractiveToolsFramework/Public/BaseGizmos/GizmoElementCircle.h b/Engine/Source/Runtime/InteractiveToolsFramework/Public/BaseGizmos/GizmoElementCircle.h
--- a/Engine/Source/Runtime/InteractiveToolsFramework/Public/BaseGizmos/GizmoElementCircle.h
+++ b/Engine/Source/Runtime/InteractiveToolsFramework/Public/BaseGizmos/GizmoElementCircle.h
@@ -3,7 +3,7 @@
 #pragma once
 
 #include "CoreMinimal.h"
-#include "BaseGizmos/GizmoElementLineBase.h"
+#include "BaseGizmos/GizmoElementCircleBase.h"
 #include "InputState.h"
 #include "UObject/ObjectMacros.h"
 #include "GizmoElementCircle.generated.h"
@@ -11,9 +11,11 @@
 /**
  * Simple object intended to be used as part of 3D Gizmos.
  * Draws a filled or line circle based on parameters.
+ * 
+ * The circle element does not yet support partial circles.
  */
 UCLASS(Transient)
-class INTERACTIVETOOLSFRAMEWORK_API UGizmoElementCircle : public UGizmoElementLineBase
+class INTERACTIVETOOLSFRAMEWORK_API UGizmoElementCircle : public UGizmoElementCircleBase
 {
 	GENERATED_BODY()
 
@@ -23,22 +25,6 @@ public:
 	virtual FInputRayHit LineTrace(const UGizmoViewContext* ViewContext, const FLineTraceTraversalState& LineTraceState, const FVector& RayOrigin, const FVector& RayDirection) override;
 	//~ End UGizmoElementBase Interface.
 
-	// Circle center.
-	virtual void SetCenter(FVector InCenter);
-	virtual FVector GetCenter() const;
-
-	// Normal to circle.
-	virtual void SetNormal(FVector InNormal);
-	virtual FVector GetNormal() const;
-
-	// Circle radius.
-	virtual void SetRadius(float Radius);
-	virtual float GetRadius() const;
-
-	// Number of sides for rendering circle.
-	virtual void SetNumSides(int InNumSides);
-	virtual int GetNumSides() const;
-
 	// Draw mesh
 	virtual void SetDrawMesh(bool InDrawMesh);
 	virtual bool GetDrawMesh() const;
@@ -57,22 +43,6 @@ public:
 
 protected:
 
-	// Circle center.
-	UPROPERTY()
-	FVector Center = FVector::ZeroVector;
-
-	// Normal to circle.
-	UPROPERTY()
-	FVector Normal = FVector::ForwardVector;
-
-	// Circle radius.
-	UPROPERTY()
-	float Radius = 100.0f;
-
-	// Number of sides for rendering circle.
-	UPROPERTY()
-	int NumSides = 64;
-
 	// Whether to render solid circle.
 	UPROPERTY()
 	bool bDrawMesh = true;
diff --git a/Engine/Source/Runtime/InteractiveToolsFramework/Public/BaseGizmos/GizmoElementCircleBase.h b/Engine/Source/Runtime/InteractiveToolsFramework/Public/BaseGizmos/GizmoElementCircleBase.h
new file mode 100644
--- /dev/null
+++ b/Engine/Source/Runtime/InteractiveToolsFramework/Public/BaseGizmos/GizmoElementCircleBase.h
@@ -0,0 +1,115 @@
+// Copyright Epic Games, Inc. All Rights Reserved.
+
+#pragma once
+
+#include "CoreMinimal.h"
+#include "BaseGizmos/GizmoElementLineBase.h"
+#include "InputState.h"
+#include "UObject/ObjectMacros.h"
+#include "GizmoElementCircleBase.generated.h"
+
+class FPrimitiveDrawInterface;
+class FMaterialRenderProxy;
+
+/**
+ * Abstract base object for circle, torus and arc.
+ */
+UCLASS(Transient, Abstract)
+class INTERACTIVETOOLSFRAMEWORK_API UGizmoElementCircleBase : public UGizmoElementLineBase
+{
+	GENERATED_BODY()
+
+public:
+
+	// CircleBase center.
+	virtual void SetCenter(const FVector& InCenter);
+	virtual FVector GetCenter() const;
+
+	// Axis0 of plane in which circle lies, must perpendicular to normal. 
+	// Start and end angles for partial circles are relative to this axis.
+	virtual void SetAxis0(const FVector& InAxis0);
+	virtual FVector GetAxis0() const;
+
+	// Axis1 of plane in which circle lies, must perpendicular to Axis0.
+	virtual void SetAxis1(const FVector& InAxis0);
+	virtual FVector GetAxis1() const;
+
+	// Circle radius.
+	virtual void SetRadius(double InRadius);
+	virtual double GetRadius() const;
+
+	// Number of segments for rendering circle.
+	virtual void SetNumSegments(int32 InNumSegments);
+	virtual int32 GetNumSegments() const;
+
+	// When Partial, renders partial arc based on angle.
+	// When PartialViewDependent, same as Partial but renders full arc when arc normal aligns with view direction
+	virtual void SetPartialType(EGizmoElementPartialType InPartial);
+	virtual EGizmoElementPartialType GetPartialType() const;
+
+	// Start of arc angle of partial torus in radians, relative to Axis0.
+	virtual void SetPartialStartAngle(double InPartialAngle);
+	virtual double GetPartialStartAngle() const;
+
+	// Start of arc angle of partial torus in radians, relative to Axis1.
+	virtual void SetPartialEndAngle(double InPartialAngle);
+	virtual double GetPartialEndAngle() const;
+
+	// If partial type is PartialViewDependent, when the cosine of angle between the normal and view direction 
+	// is within this tolerance, the arc will be rendered as full rather than partial
+	virtual void SetPartialViewDependentMaxCosTol(double InPartialViewDependentMaxCosTol);
+	virtual double GetPartialViewDependentMaxCosTol() const;
+
+protected:
+
+	// CircleBase center.
+	UPROPERTY()
+	FVector Center = FVector::ZeroVector;
+
+	// Axis0 of plane in which circle lies, must perpendicular to normal. 
+	// Start and end angles for partial circles are relative to this axis.
+	UPROPERTY()
+	FVector Axis0 = FVector::RightVector;
+
+	// Axis1 of plane in which circle lies, must perpendicular to Axis0.
+	UPROPERTY()
+	FVector Axis1 = FVector::ForwardVector;
+
+	// Radius of main circle, some derived elements have inner radius (e.g. torus and arc)
+	UPROPERTY()
+	double Radius = 100.0;
+
+	// Number of segments for rendering arc.
+	UPROPERTY()
+	int32 NumSegments = 64;
+
+	// True when the arc is not full.
+	UPROPERTY()
+	EGizmoElementPartialType PartialType = EGizmoElementPartialType::None;
+
+	// Start angle to render for partial torus
+	UPROPERTY()
+	double PartialStartAngle = 0.0;
+
+	// End angle to render for partial torus
+	UPROPERTY()
+	double PartialEndAngle = UE_DOUBLE_TWO_PI;
+
+	// For PartialViewDependent, max cosine of angle between the normal and view direction 
+	// Within this tolerance, the arc will be rendered as full rather than partial
+	UPROPERTY()
+	double PartialViewDependentMaxCosTol = 0.96f;
+
+protected:
+
+	// Returns whether element should be partial based on current view
+	bool IsPartial(const FSceneView* View, const FVector& InWorldCenter, const FVector& InWorldNormal);
+
+	// Returns whether element should be partial based on current view
+	bool IsPartial(const UGizmoViewContext* ViewContext, const FVector& InWorldCenter, const FVector& InWorldNormal);
+
+	// Returns whether element should be partial based on current view
+	bool IsPartial(const FVector& InWorldCenter, const FVector& InWorldNormal,
+		const FVector& InViewLocation, const FVector& InViewDirection,
+		bool bIsPerspectiveProjection);
+};
\ No newline at end of file
diff --git a/Engine/Source/Runtime/InteractiveToolsFramework/Public/BaseGizmos/GizmoElementShapes.h b/Engine/Source/Runtime/InteractiveToolsFramework/Public/BaseGizmos/GizmoElementShapes.h
--- a/Engine/Source/Runtime/InteractiveToolsFramework/Public/BaseGizmos/GizmoElementShapes.h
+++ b/Engine/Source/Runtime/InteractiveToolsFramework/Public/BaseGizmos/GizmoElementShapes.h
@@ -2,6 +2,7 @@
 
 #pragma once
 
+#include "BaseGizmos/GizmoElementArc.h"
 #include "BaseGizmos/GizmoElementArrow.h"
 #include "BaseGizmos/GizmoElementBox.h"
 #include "BaseGizmos/GizmoElementCone.h"
diff --git a/Engine/Source/Runtime/InteractiveToolsFramework/Public/BaseGizmos/GizmoElementShared.h b/Engine/Source/Runtime/InteractiveToolsFramework/Public/BaseGizmos/GizmoElementShared.h
--- a/Engine/Source/Runtime/InteractiveToolsFramework/Public/BaseGizmos/GizmoElementShared.h
+++ b/Engine/Source/Runtime/InteractiveToolsFramework/Public/BaseGizmos/GizmoElementShared.h
@@ -62,3 +62,18 @@ enum class EGizmoElementViewAlignType
 	Axial
 };
 
+
+//
+// Partial type: render partial element for those elements which support it.
+// 
+//   Partial				- Render partial element.
+//   PartialViewDependent   - Render partial unless view direction aligns with an axis or normal specified by the element type.
+//
+UENUM()
+enum class EGizmoElementPartialType
+{
+	None,
+	Partial,
+	PartialViewDependent
+};
+
diff --git a/Engine/Source/Runtime/InteractiveToolsFramework/Public/BaseGizmos/GizmoElementTorus.h b/Engine/Source/Runtime/InteractiveToolsFramework/Public/BaseGizmos/GizmoElementTorus.h
--- a/Engine/Source/Runtime/InteractiveToolsFramework/Public/BaseGizmos/GizmoElementTorus.h
+++ b/Engine/Source/Runtime/InteractiveToolsFramework/Public/BaseGizmos/GizmoElementTorus.h
@@ -3,7 +3,7 @@
 #pragma once
 
 #include "CoreMinimal.h"
-#include "BaseGizmos/GizmoElementBase.h"
+#include "BaseGizmos/GizmoElementCircleBase.h"
 #include "InputState.h"
 #include "UObject/ObjectMacros.h"
 #include "GizmoElementTorus.generated.h"
@@ -19,7 +19,7 @@ class FMaterialRenderProxy;
  * See comment above LineTrace method below for details of how this intersection is approximated.
  */
 UCLASS(Transient)
-class INTERACTIVETOOLSFRAMEWORK_API UGizmoElementTorus : public UGizmoElementBase
+class INTERACTIVETOOLSFRAMEWORK_API UGizmoElementTorus : public UGizmoElementCircleBase
 {
 	GENERATED_BODY()
 
@@ -35,41 +35,13 @@ public:
 	virtual FInputRayHit LineTrace(const UGizmoViewContext* ViewContext, const FLineTraceTraversalState& LineTraceState, const FVector& RayOrigin, const FVector& RayDirection) override;
 	//~ End UGizmoElementBase Interface.
 
-	// Torus center.
-	virtual void SetCenter(const FVector& InCenter);
-	virtual FVector GetCenter() const;
-
-	// Normal to torus.
-	virtual void SetNormal(const FVector& InNormal);
-	virtual FVector GetNormal() const;
-
-	// Begin axis, should be perpendicular to normal, this where partial torus will begin.
-	virtual void SetBeginAxis(const FVector& InBeginAxis);
-	virtual FVector GetBeginAxis() const;
-
-	// Outer circle radius.
-	virtual void SetOuterRadius(float InOuterRadius);
-	virtual float GetOuterRadius() const;
-
 	// Inner circles radius.
-	virtual void SetInnerRadius(float InInnerRadius);
-	virtual float GetInnerRadius() const;
-
-	// Number of outer segments for rendering torus.
-	virtual void SetOuterSegments(int32 InOuterSegments);
-	virtual int32 GetOuterSegments() const;
+	virtual void SetInnerRadius(double InInnerRadius);
+	virtual double GetInnerRadius() const;
 
 	// Number of inner slices for rendering torus.
-	virtual void SetInnerSlices(int32 InInnerSlices);
-	virtual int32 GetInnerSlices() const;
-
-	// True when the torus is a partial arc.
-	virtual void SetPartial(bool InPartial);
-	virtual bool GetPartial() const;
-
-	// If partial, arc angle of partial torus in radians.
-	virtual void SetAngle(float InAngle);
-	virtual float GetAngle() const;
+	virtual void SetNumInnerSlices(int32 InNumInnerSlices);
+	virtual int32 GetNumInnerSlices() const;
 
 	// If partial, renders end caps when true.
 	virtual void SetEndCaps(bool InEndCaps);
@@ -77,47 +49,20 @@ public:
 
 protected:
 
-	// Torus center.
-	UPROPERTY()
-	FVector Center = FVector::ZeroVector;
-
-	// Normal to plane in which torus lies.
-	UPROPERTY()
-	FVector Normal = FVector::UpVector;
-
-	// Plane axis, should be perpendicular to normal. Indicates where partial torus should begin
-	UPROPERTY()
-	FVector BeginAxis = FVector::ForwardVector;
-
-	// Torus outer radius.
-	UPROPERTY()
-	float OuterRadius = 100.0f;
-
 	// Torus inner radius.
 	UPROPERTY()
-	float InnerRadius = 5.0f;
-
-	// Number of segments for rendering torus.
-	UPROPERTY()
-	int32 OuterSegments = 64;
+	double InnerRadius = 5.0;
 
 	// Number of slices to render in each torus segment.
 	UPROPERTY()
-	int32 InnerSlices = 8;
+	int32 NumInnerSlices = 8;
 
-	// True when the torus is not full.
+	// Whether to render end caps on a partial torus.
 	UPROPERTY()
-	bool bPartial = false;
+	bool bEndCaps = false;
 
-	// True when partial torus should be aligned to the screen
-	UPROPERTY(EditAnywhere, Category = Options, meta = (EditCondition = "bPartial"))
-	bool bScreenAlignPartial = false;
+protected:
 
-	// Angle to render for partial torus
-	UPROPERTY(EditAnywhere, Category = Options, meta = (EditCondition = "bPartial"))
-	float Angle = 2.0f * PI;
+	bool IsGlancingAngle(const UGizmoViewContext* View, const FVector& InWorldCenter, const FVector& InWorldNormal);
 
-	// Whether to render end caps on a partial torus.
-	UPROPERTY(EditAnywhere, Category = Options, meta = (EditCondition = "bPartial"))
-	bool bEndCaps = false;
 };
\ No newline at end of file
