make brush size hotkeys apply to the brush world radius when in the UI mode that directly uses that radius

diff --git a/Engine/Source/Runtime/InteractiveToolsFramework/Private/BaseTools/BaseBrushTool.cpp b/Engine/Source/Runtime/InteractiveToolsFramework/Private/BaseTools/BaseBrushTool.cpp
--- a/Engine/Source/Runtime/InteractiveToolsFramework/Private/BaseTools/BaseBrushTool.cpp
+++ b/Engine/Source/Runtime/InteractiveToolsFramework/Private/BaseTools/BaseBrushTool.cpp
@@ -61,14 +61,29 @@ void UBaseBrushTool::OnPropertyModified(UObject* PropertySet, FProperty* Propert
 
 void UBaseBrushTool::IncreaseBrushSizeAction()
 {
-	BrushProperties->BrushSize = FMath::Clamp(BrushProperties->BrushSize + 0.025f, 0.0f, 1.0f);
+	if (BrushProperties->bSpecifyRadius)
+	{
+		// Hardcoded max of 1000 chosen to match the BrushRadius "UIMax" specified in UBrushBaseProperties
+		BrushProperties->BrushRadius = FMath::Min(BrushProperties->BrushRadius * 1.1, 1000);
+	}
+	else
+	{
+		BrushProperties->BrushSize = FMath::Clamp(BrushProperties->BrushSize + 0.025f, 0.0f, 1.0f);
+	}
 	RecalculateBrushRadius();
 	NotifyOfPropertyChangeByTool(BrushProperties);
 }
 
 void UBaseBrushTool::DecreaseBrushSizeAction()
 {
-	BrushProperties->BrushSize = FMath::Clamp(BrushProperties->BrushSize - 0.025f, 0.0f, 1.0f);
+	if (BrushProperties->bSpecifyRadius)
+	{
+		BrushProperties->BrushRadius = FMath::Max(BrushProperties->BrushRadius / 1.1, 1);
+	}
+	else
+	{
+		BrushProperties->BrushSize = FMath::Clamp(BrushProperties->BrushSize - 0.025f, 0.0f, 1.0f);
+	}
 	RecalculateBrushRadius();
 	NotifyOfPropertyChangeByTool(BrushProperties);
 }
@@ -166,13 +181,15 @@ void UBaseBrushTool::RegisterActions(FInteractiveToolActionSet& ActionSet)
 void UBaseBrushTool::RecalculateBrushRadius()
 {
 	TInterval<float> ScaledBrushSizeRange(BrushRelativeSizeRange.Min/WorldToLocalScale, BrushRelativeSizeRange.Max/WorldToLocalScale);
-	CurrentBrushRadius = 0.5 * ScaledBrushSizeRange.Interpolate(BrushProperties->BrushSize);
+	
 	if (BrushProperties->bSpecifyRadius)
 	{
 		CurrentBrushRadius = BrushProperties->BrushRadius;
+		BrushProperties->BrushSize = (2 * CurrentBrushRadius - ScaledBrushSizeRange.Min) / ScaledBrushSizeRange.Size();
 	}
 	else
 	{
+		CurrentBrushRadius = 0.5 * ScaledBrushSizeRange.Interpolate(BrushProperties->BrushSize);
 		BrushProperties->BrushRadius = CurrentBrushRadius;
 	}
 }
