Gizmo element library: update recently added hover/interact vertex color state when updating state attributes.

diff --git a/Engine/Source/Runtime/InteractiveToolsFramework/Private/BaseGizmos/GizmoElementRenderState.cpp b/Engine/Source/Runtime/InteractiveToolsFramework/Private/BaseGizmos/GizmoElementRenderState.cpp
--- a/Engine/Source/Runtime/InteractiveToolsFramework/Private/BaseGizmos/GizmoElementRenderState.cpp
+++ b/Engine/Source/Runtime/InteractiveToolsFramework/Private/BaseGizmos/GizmoElementRenderState.cpp
@@ -99,6 +99,8 @@ void FGizmoElementMeshRenderStateAttributes::Update(FGizmoElementMeshRenderState
 	HoverMaterial.UpdateState(InChildAttributes.HoverMaterial);
 	InteractMaterial.UpdateState(InChildAttributes.InteractMaterial);
 	VertexColor.UpdateState(InChildAttributes.VertexColor);
+	HoverVertexColor.UpdateState(InChildAttributes.HoverVertexColor);
+	InteractVertexColor.UpdateState(InChildAttributes.InteractVertexColor);
 }
 
 FLinearColor FGizmoElementLineRenderStateAttributes::GetLineColor(EGizmoElementInteractionState InteractionState)
