TransformProxy - fix for black geometry when using the transform gizmo to flip the sign of the scale

diff --git a/Engine/Source/Runtime/InteractiveToolsFramework/Private/BaseGizmos/TransformProxy.cpp b/Engine/Source/Runtime/InteractiveToolsFramework/Private/BaseGizmos/TransformProxy.cpp
--- a/Engine/Source/Runtime/InteractiveToolsFramework/Private/BaseGizmos/TransformProxy.cpp
+++ b/Engine/Source/Runtime/InteractiveToolsFramework/Private/BaseGizmos/TransformProxy.cpp
@@ -15,7 +15,13 @@ void UTransformProxy::AddComponent(USceneComponent* Component, bool bModifyCompo
 	NewObj.Component = Component;
 	NewObj.bModifyComponentOnTransform = bModifyComponentOnTransform;
 	NewObj.GetTransformFunc = [Component]() { return Component->GetComponentToWorld(); };
-	NewObj.SetTransformFunc = [Component](FTransform NewTransform) { return Component->SetWorldTransform(NewTransform); };
+	NewObj.SetTransformFunc = [Component](FTransform NewTransform) 
+										 { 
+											Component->SetWorldTransform(NewTransform); 
+											// mimic behavior in USceneComponent::PostEditChangeProperty that is triggered by updating the scale3 in the transform details panel.
+											// otherwise changes in the sign of the scale will result rendering artifacts 
+											Component->InvalidateLightingCacheDetailed(true, false); 
+										 };
 	NewObj.UserDefinedIndex = 0;
 	NewObj.StartTransform = NewObj.GetTransformFunc();
 	NewObj.RelativeTransform = FTransform::Identity;
