MeshModelingTools: Align the sensitivity of plane and uniform scaling gizmos.

diff --git a/Engine/Source/Runtime/InteractiveToolsFramework/Public/BaseGizmos/ParameterToTransformAdapters.h b/Engine/Source/Runtime/InteractiveToolsFramework/Public/BaseGizmos/ParameterToTransformAdapters.h
--- a/Engine/Source/Runtime/InteractiveToolsFramework/Public/BaseGizmos/ParameterToTransformAdapters.h
+++ b/Engine/Source/Runtime/InteractiveToolsFramework/Public/BaseGizmos/ParameterToTransformAdapters.h
@@ -480,7 +480,7 @@ public:
 
 	/** Coordinate delta is multiplied by this amount */
 	UPROPERTY()
-	float ScaleMultiplier = 0.025f;
+	float ScaleMultiplier = 0.02f;
 
 public:
 	/** Parameter is the two line-equation parameters that this Vec2ParameterSource provides */
@@ -708,7 +708,7 @@ public:
 
 		if (bUseEqualScaling)
 		{
-			Delta = FVector2D(FMath::Min(Delta.X, Delta.Y));
+			Delta = FVector2D(Delta.X + Delta.Y);
 		}
 
 		FTransform NewTransform = InitialTransform;
