InteractiveToolsFramework : fix type cast warnings.

diff --git a/Engine/Source/Runtime/InteractiveToolsFramework/Private/BaseBehaviors/SingleClickOrDragBehavior.cpp b/Engine/Source/Runtime/InteractiveToolsFramework/Private/BaseBehaviors/SingleClickOrDragBehavior.cpp
--- a/Engine/Source/Runtime/InteractiveToolsFramework/Private/BaseBehaviors/SingleClickOrDragBehavior.cpp
+++ b/Engine/Source/Runtime/InteractiveToolsFramework/Private/BaseBehaviors/SingleClickOrDragBehavior.cpp
@@ -68,7 +68,7 @@ FInputCaptureUpdate USingleClickOrDragInputBehavior::UpdateCapture(const FInputD
 	// check if mouse has moved far enough that we want to swap to drag behavior
 	if (bInDrag == false)
 	{
-		float MouseMovement = FVector2D::Distance(Input.Mouse.Position2D, MouseDownPosition2D);
+		float MouseMovement = static_cast<float>(FVector2D::Distance(Input.Mouse.Position2D, MouseDownPosition2D));
 		if (MouseMovement > ClickDistanceThreshold)
 		{
 			FInputDeviceState StartInput = Input;
diff --git a/Engine/Source/Runtime/InteractiveToolsFramework/Private/BaseGizmos/AxisPositionGizmo.cpp b/Engine/Source/Runtime/InteractiveToolsFramework/Private/BaseGizmos/AxisPositionGizmo.cpp
--- a/Engine/Source/Runtime/InteractiveToolsFramework/Private/BaseGizmos/AxisPositionGizmo.cpp
+++ b/Engine/Source/Runtime/InteractiveToolsFramework/Private/BaseGizmos/AxisPositionGizmo.cpp
@@ -69,7 +69,7 @@ void UAxisPositionGizmo::OnClickPress(const FInputDeviceRay& PressPos)
 
 	FVector AxisOrigin = AxisSource->GetOrigin();
 
-	float DirectionSign = FVector::DotProduct(InteractionStartPoint - AxisOrigin, InteractionAxis);
+	double DirectionSign = FVector::DotProduct(InteractionStartPoint - AxisOrigin, InteractionAxis);
 	ParameterSign = (bEnableSignedAxis && DirectionSign < 0) ? -1.0f : 1.0f;
 
 	// Figure out how the parameter would need to be adjusted to bring the axis origin to the
diff --git a/Engine/Source/Runtime/InteractiveToolsFramework/Private/BaseGizmos/GizmoArrowComponent.cpp b/Engine/Source/Runtime/InteractiveToolsFramework/Private/BaseGizmos/GizmoArrowComponent.cpp
--- a/Engine/Source/Runtime/InteractiveToolsFramework/Private/BaseGizmos/GizmoArrowComponent.cpp
+++ b/Engine/Source/Runtime/InteractiveToolsFramework/Private/BaseGizmos/GizmoArrowComponent.cpp
@@ -112,7 +112,7 @@ public:
 					(HoverThicknessMultiplier * Thickness) : (Thickness);
 				if (View->IsPerspectiveProjection())
 				{
-					UseThickness *= (View->FOV / 90.0);		// compensate for FOV scaling in Gizmos...
+					UseThickness *= (View->FOV / 90.0f);		// compensate for FOV scaling in Gizmos...
 				}
 
 				PDI->DrawLine(StartPoint, EndPoint, Color, SDPG_Foreground, UseThickness, 0.0f, true);
@@ -136,8 +136,8 @@ public:
 		return false;
 	}
 
-	virtual uint32 GetMemoryFootprint(void) const override { return sizeof *this + GetAllocatedSize(); }
-	uint32 GetAllocatedSize(void) const { return FPrimitiveSceneProxy::GetAllocatedSize(); }
+	virtual uint32 GetMemoryFootprint(void) const override { return IntCastChecked<uint32>(sizeof *this + GetAllocatedSize()); }
+	SIZE_T GetAllocatedSize(void) const { return FPrimitiveSceneProxy::GetAllocatedSize(); }
 
 	void SetExternalHoverState(bool* HoverState)
 	{
@@ -211,7 +211,7 @@ bool UGizmoArrowComponent::LineTraceComponent(FHitResult& OutHit, const FVector
 	}
 
 	OutHit.Component = this;
-	OutHit.Distance = FVector::Distance(Start, NearestLine);
+	OutHit.Distance = static_cast<float>(FVector::Distance(Start, NearestLine));
 	OutHit.ImpactPoint = NearestLine;
 	return true;
 }
diff --git a/Engine/Source/Runtime/InteractiveToolsFramework/Private/BaseGizmos/GizmoBoxComponent.cpp b/Engine/Source/Runtime/InteractiveToolsFramework/Private/BaseGizmos/GizmoBoxComponent.cpp
--- a/Engine/Source/Runtime/InteractiveToolsFramework/Private/BaseGizmos/GizmoBoxComponent.cpp
+++ b/Engine/Source/Runtime/InteractiveToolsFramework/Private/BaseGizmos/GizmoBoxComponent.cpp
@@ -130,7 +130,7 @@ public:
 					(HoverThicknessMultiplier * Thickness) : (Thickness);
 				if (!bIsOrtho)
 				{
-					UseThickness *= (View->FOV / 90.0);		// compensate for FOV scaling in Gizmos...
+					UseThickness *= (View->FOV / 90.0f);		// compensate for FOV scaling in Gizmos...
 				}
 
 				double DimensionX = LengthScale * Dimensions.X * 0.5f;
@@ -153,11 +153,11 @@ public:
 						View->IsPerspectiveProjection() ? WorldOrigin - View->ViewLocation : View->GetViewDirection();
 
 					// find box corner direction that is most aligned with view direction. That's the corner we will hide.
-					float MaxDot = -999999;
+					double MaxDot = -999999.;
 					int MaxDotIndex = -1;
 					for (int j = 0; j < 8; ++j)
 					{
-						float Dot = FVector::DotProduct(Points[j], ViewDirection);
+						double Dot = FVector::DotProduct(Points[j], ViewDirection);
 						if (Dot > MaxDot)
 						{
 							MaxDot = Dot;
@@ -205,8 +205,8 @@ public:
 		return false;
 	}
 
-	virtual uint32 GetMemoryFootprint(void) const override { return sizeof *this + GetAllocatedSize(); }
-	uint32 GetAllocatedSize(void) const { return FPrimitiveSceneProxy::GetAllocatedSize(); }
+	virtual uint32 GetMemoryFootprint(void) const override { return IntCastChecked<uint32>(sizeof *this + GetAllocatedSize()); }
+	SIZE_T GetAllocatedSize(void) const { return FPrimitiveSceneProxy::GetAllocatedSize(); }
 
 	void SetExternalHoverState(bool* HoverState)
 	{
@@ -330,7 +330,7 @@ bool UGizmoBoxComponent::LineTraceComponent(FHitResult& OutHit, const FVector St
 
 	OutHit.Component = this;
 	OutHit.ImpactPoint = HitWorld;
-	OutHit.Distance = FVector::Distance(Start, HitWorld);
+	OutHit.Distance = static_cast<float>(FVector::Distance(Start, HitWorld));
 	//OutHit.ImpactNormal = ;
 	
 	return true;
diff --git a/Engine/Source/Runtime/InteractiveToolsFramework/Private/BaseGizmos/GizmoCircleComponent.cpp b/Engine/Source/Runtime/InteractiveToolsFramework/Private/BaseGizmos/GizmoCircleComponent.cpp
--- a/Engine/Source/Runtime/InteractiveToolsFramework/Private/BaseGizmos/GizmoCircleComponent.cpp
+++ b/Engine/Source/Runtime/InteractiveToolsFramework/Private/BaseGizmos/GizmoCircleComponent.cpp
@@ -82,8 +82,8 @@ public:
 					(HoverThicknessMultiplier * Thickness) : (Thickness);
 				if (!bIsOrtho)
 				{
-					UseThickness *= (View->FOV / 90.0);		// compensate for FOV scaling in Gizmos...
-					BackThickness *= (View->FOV / 90.0);		// compensate for FOV scaling in Gizmos...
+					UseThickness *= (View->FOV / 90.0f);		// compensate for FOV scaling in Gizmos...
+					BackThickness *= (View->FOV / 90.0f);		// compensate for FOV scaling in Gizmos...
 				}
 
 				if (bDrawFullCircle)
@@ -193,8 +193,8 @@ public:
 		return false;
 	}
 
-	virtual uint32 GetMemoryFootprint(void) const override { return sizeof *this + GetAllocatedSize(); }
-	uint32 GetAllocatedSize(void) const { return FPrimitiveSceneProxy::GetAllocatedSize(); }
+	virtual uint32 GetMemoryFootprint(void) const override { return IntCastChecked<uint32>(sizeof *this + GetAllocatedSize()); }
+	SIZE_T GetAllocatedSize(void) const { return FPrimitiveSceneProxy::GetAllocatedSize(); }
 
 	void SetExternalHoverState(bool* HoverState)
 	{
@@ -272,7 +272,7 @@ bool UGizmoCircleComponent::LineTraceComponent(FHitResult& OutHit, const FVector
 	}
 
 	float LengthScale = PixelToWorldScale;
-	double UseRadius = LengthScale * Radius;
+	float UseRadius = LengthScale * Radius;
 
 
 	FRay Ray(Start, End - Start, false);
@@ -317,7 +317,7 @@ bool UGizmoCircleComponent::LineTraceComponent(FHitResult& OutHit, const FVector
 	}
 
 	OutHit.Component = this;
-	OutHit.Distance = FVector::Distance(Start, NearestRay);
+	OutHit.Distance = static_cast<float>(FVector::Distance(Start, NearestRay));
 	OutHit.ImpactPoint = NearestRay;
 	return true;
 }
diff --git a/Engine/Source/Runtime/InteractiveToolsFramework/Private/BaseGizmos/GizmoElementArc.cpp b/Engine/Source/Runtime/InteractiveToolsFramework/Private/BaseGizmos/GizmoElementArc.cpp
--- a/Engine/Source/Runtime/InteractiveToolsFramework/Private/BaseGizmos/GizmoElementArc.cpp
+++ b/Engine/Source/Runtime/InteractiveToolsFramework/Private/BaseGizmos/GizmoElementArc.cpp
@@ -21,7 +21,7 @@ static void DrawThickArc(FPrimitiveDrawInterface* PDI, const FVector& InCenter,
 		return;
 	}
 
-	const int32 NumPoints = FMath::TruncToInt(InNumSegments * (InEndAngle - InStartAngle) / (UE_DOUBLE_PI / 2.0)) + 1;
+	const int32 NumPoints = FMath::TruncToInt32(InNumSegments * (InEndAngle - InStartAngle) / (UE_DOUBLE_PI / 2.0)) + 1;
 
 	FColor TriangleColor = InColor;
 	FColor RingColor = InColor;
@@ -47,7 +47,7 @@ static void DrawThickArc(FPrimitiveDrawInterface* PDI, const FVector& InCenter,
 			VertexDir.Normalize();
 
 			double TCAngle = Percent * (PI / 2);
-			FVector2f TC(TCRadius * FMath::Cos(Angle), TCRadius * FMath::Sin(Angle));
+			FVector2f TC(static_cast<float>(TCRadius * FMath::Cos(Angle)), static_cast<float>(TCRadius * FMath::Sin(Angle)));
 
 			// Keep the vertices in local space so that we don't lose precision when dealing with LWC
 			// The local-to-world transform is handled in the MeshBuilder.Draw() call at the end of this function
@@ -207,7 +207,7 @@ FInputRayHit UGizmoElementArc::LineTrace(const UGizmoViewContext* ViewContext, c
 
 		if (HitDepth >= 0.0)
 		{
-			FInputRayHit RayHit(HitDepth);
+			FInputRayHit RayHit(static_cast<float>(HitDepth));
 			RayHit.SetHitObject(this);
 			RayHit.HitIdentifier = PartIdentifier;
 			return RayHit;
diff --git a/Engine/Source/Runtime/InteractiveToolsFramework/Private/BaseGizmos/GizmoElementBox.cpp b/Engine/Source/Runtime/InteractiveToolsFramework/Private/BaseGizmos/GizmoElementBox.cpp
--- a/Engine/Source/Runtime/InteractiveToolsFramework/Private/BaseGizmos/GizmoElementBox.cpp
+++ b/Engine/Source/Runtime/InteractiveToolsFramework/Private/BaseGizmos/GizmoElementBox.cpp
@@ -48,7 +48,7 @@ FInputRayHit UGizmoElementBox::LineTrace(const UGizmoViewContext* ViewContext, c
 		UE::Geometry::TOrientedBox3<double> Box(Frame, WorldExtent);
 		if (UE::Geometry::TIntrRay3OrientedBox3<double>::FindIntersection(Ray, Box, HitDepth))
 		{
-			FInputRayHit RayHit(HitDepth);
+			FInputRayHit RayHit(static_cast<float>(HitDepth));
 			RayHit.SetHitObject(this);
 			RayHit.HitIdentifier = PartIdentifier;
 			return RayHit;
diff --git a/Engine/Source/Runtime/InteractiveToolsFramework/Private/BaseGizmos/GizmoElementCircle.cpp b/Engine/Source/Runtime/InteractiveToolsFramework/Private/BaseGizmos/GizmoElementCircle.cpp
--- a/Engine/Source/Runtime/InteractiveToolsFramework/Private/BaseGizmos/GizmoElementCircle.cpp
+++ b/Engine/Source/Runtime/InteractiveToolsFramework/Private/BaseGizmos/GizmoElementCircle.cpp
@@ -17,7 +17,7 @@ void UGizmoElementCircle::Render(IToolsContextRenderAPI* RenderAPI, const FRende
 
 	if (bVisibleViewDependent)
 	{
-		const float WorldRadius = Radius * CurrentRenderState.LocalToWorldTransform.GetScale3D().X;
+		const double WorldRadius = Radius * CurrentRenderState.LocalToWorldTransform.GetScale3D().X;
 		const FVector WorldCenter = CurrentRenderState.LocalToWorldTransform.TransformPosition(FVector::ZeroVector);
 		const FVector WorldAxis0 = CurrentRenderState.LocalToWorldTransform.TransformVectorNoScale(Axis0);
 		const FVector WorldAxis1 = CurrentRenderState.LocalToWorldTransform.TransformVectorNoScale(Axis1);
@@ -74,7 +74,7 @@ FInputRayHit UGizmoElementCircle::LineTrace(const UGizmoViewContext* ViewContext
 
 			if (Result.intersects)
 			{
-				FInputRayHit RayHit(Result.parameter.Min);
+				FInputRayHit RayHit(static_cast<float>(Result.parameter.Min));
 				RayHit.SetHitObject(this);
 				RayHit.HitIdentifier = PartIdentifier;
 				return RayHit;
@@ -92,7 +92,7 @@ FInputRayHit UGizmoElementCircle::LineTrace(const UGizmoViewContext* ViewContext
 			FVector HitPoint = RayOrigin + RayDirection * HitDepth;
 
 			FVector NearestCirclePos;
-			GizmoMath::ClosetPointOnCircle(HitPoint, WorldCenter, WorldNormal, WorldRadius, NearestCirclePos);
+			GizmoMath::ClosetPointOnCircle(HitPoint, WorldCenter, WorldNormal, static_cast<float>(WorldRadius), NearestCirclePos);
 
 			FRay Ray(RayOrigin, RayDirection, true);
 			FVector NearestRayPos = Ray.ClosestPoint(NearestCirclePos);
@@ -102,7 +102,7 @@ FInputRayHit UGizmoElementCircle::LineTrace(const UGizmoViewContext* ViewContext
 			
 			if (Distance <= HitBuffer)
 			{
-				FInputRayHit RayHit(HitDepth);
+				FInputRayHit RayHit(static_cast<float>(HitDepth));
 				RayHit.SetHitObject(this);
 				RayHit.HitIdentifier = PartIdentifier;
 				return RayHit;
diff --git a/Engine/Source/Runtime/InteractiveToolsFramework/Private/BaseGizmos/GizmoElementCone.cpp b/Engine/Source/Runtime/InteractiveToolsFramework/Private/BaseGizmos/GizmoElementCone.cpp
--- a/Engine/Source/Runtime/InteractiveToolsFramework/Private/BaseGizmos/GizmoElementCone.cpp
+++ b/Engine/Source/Runtime/InteractiveToolsFramework/Private/BaseGizmos/GizmoElementCone.cpp
@@ -57,7 +57,7 @@ FInputRayHit UGizmoElementCone::LineTrace(const UGizmoViewContext* ViewContext,
 
 		if (bIntersects)
 		{
-			FInputRayHit RayHit(RayParam);
+			FInputRayHit RayHit(static_cast<float>(RayParam));
 			RayHit.SetHitObject(this);
 			RayHit.HitIdentifier = PartIdentifier;
 			return RayHit;
diff --git a/Engine/Source/Runtime/InteractiveToolsFramework/Private/BaseGizmos/GizmoElementCylinder.cpp b/Engine/Source/Runtime/InteractiveToolsFramework/Private/BaseGizmos/GizmoElementCylinder.cpp
--- a/Engine/Source/Runtime/InteractiveToolsFramework/Private/BaseGizmos/GizmoElementCylinder.cpp
+++ b/Engine/Source/Runtime/InteractiveToolsFramework/Private/BaseGizmos/GizmoElementCylinder.cpp
@@ -57,7 +57,7 @@ FInputRayHit UGizmoElementCylinder::LineTrace(const UGizmoViewContext* ViewConte
 
 		if (bIntersects)
 		{
-			FInputRayHit RayHit(RayParam);
+			FInputRayHit RayHit(static_cast<float>(RayParam));
 			RayHit.SetHitObject(this);
 			RayHit.HitIdentifier = PartIdentifier;
 			return RayHit;
diff --git a/Engine/Source/Runtime/InteractiveToolsFramework/Private/BaseGizmos/GizmoElementGroup.cpp b/Engine/Source/Runtime/InteractiveToolsFramework/Private/BaseGizmos/GizmoElementGroup.cpp
--- a/Engine/Source/Runtime/InteractiveToolsFramework/Private/BaseGizmos/GizmoElementGroup.cpp
+++ b/Engine/Source/Runtime/InteractiveToolsFramework/Private/BaseGizmos/GizmoElementGroup.cpp
@@ -5,9 +5,9 @@
 
 #include UE_INLINE_GENERATED_CPP_BY_NAME(GizmoElementGroup)
 
-void UGizmoElementGroup::ApplyUniformConstantScaleToTransform(float PixelToWorldScale, FTransform& InOutLocalToWorldTransform) const
+void UGizmoElementGroup::ApplyUniformConstantScaleToTransform(double PixelToWorldScale, FTransform& InOutLocalToWorldTransform) const
 {
-	float Scale = InOutLocalToWorldTransform.GetScale3D().X;
+	double Scale = InOutLocalToWorldTransform.GetScale3D().X;
 	if (bConstantScale)
 	{
 		Scale *= PixelToWorldScale;
diff --git a/Engine/Source/Runtime/InteractiveToolsFramework/Private/BaseGizmos/GizmoElementLineBase.cpp b/Engine/Source/Runtime/InteractiveToolsFramework/Private/BaseGizmos/GizmoElementLineBase.cpp
--- a/Engine/Source/Runtime/InteractiveToolsFramework/Private/BaseGizmos/GizmoElementLineBase.cpp
+++ b/Engine/Source/Runtime/InteractiveToolsFramework/Private/BaseGizmos/GizmoElementLineBase.cpp
@@ -32,7 +32,7 @@ float UGizmoElementLineBase::GetCurrentLineThickness(bool bPerspectiveView, floa
 
 	if (bPerspectiveView)
 	{
-		CurrentLineThickness *= (InViewFOV / 90.0);		// compensate for FOV scaling in Gizmos...
+		CurrentLineThickness *= (InViewFOV / 90.0f);		// compensate for FOV scaling in Gizmos...
 	}
 
 	return CurrentLineThickness;
@@ -40,10 +40,10 @@ float UGizmoElementLineBase::GetCurrentLineThickness(bool bPerspectiveView, floa
 
 void UGizmoElementLineBase::SetLineThickness(float InLineThickness)
 {
-	if (InLineThickness < 0.0)
+	if (InLineThickness < 0.0f)
 	{
 		UE_LOG(LogGizmoElementLineBase, Warning, TEXT("Invalid gizmo element line thickness %f, will be set to 0.0."), InLineThickness);
-		LineThickness = 0.0;
+		LineThickness = 0.0f;
 	}
 	else
 	{
diff --git a/Engine/Source/Runtime/InteractiveToolsFramework/Private/BaseGizmos/GizmoElementRectangle.cpp b/Engine/Source/Runtime/InteractiveToolsFramework/Private/BaseGizmos/GizmoElementRectangle.cpp
--- a/Engine/Source/Runtime/InteractiveToolsFramework/Private/BaseGizmos/GizmoElementRectangle.cpp
+++ b/Engine/Source/Runtime/InteractiveToolsFramework/Private/BaseGizmos/GizmoElementRectangle.cpp
@@ -20,8 +20,8 @@ void UGizmoElementRectangle::Render(IToolsContextRenderAPI* RenderAPI, const FRe
 		const FVector WorldUpAxis = CurrentRenderState.LocalToWorldTransform.TransformVectorNoScale(UpDirection);
 		const FVector WorldSideAxis = CurrentRenderState.LocalToWorldTransform.TransformVectorNoScale(SideDirection);
 		const FVector WorldCenter = CurrentRenderState.LocalToWorldTransform.TransformPosition(FVector::ZeroVector);
-		const float WorldWidth = Width * CurrentRenderState.LocalToWorldTransform.GetScale3D().X;
-		const float WorldHeight = Height * CurrentRenderState.LocalToWorldTransform.GetScale3D().X;
+		const float WorldWidth = static_cast<float>(Width * CurrentRenderState.LocalToWorldTransform.GetScale3D().X);
+		const float WorldHeight = static_cast<float>(Height * CurrentRenderState.LocalToWorldTransform.GetScale3D().X);
 
 		FPrimitiveDrawInterface* PDI = RenderAPI->GetPrimitiveDrawInterface();
 
@@ -87,7 +87,7 @@ FInputRayHit UGizmoElementRectangle::LineTrace(const UGizmoViewContext* ViewCont
 			// clip to rectangle dimensions
 			if (HdU >= 0.0 && HdU <= WorldHeight && HdS >= 0.0 && HdS <= WorldWidth)
 			{
-				FInputRayHit RayHit(HitDepth);
+				FInputRayHit RayHit(static_cast<float>(HitDepth));
 				RayHit.SetHitObject(this);
 				RayHit.HitIdentifier = PartIdentifier;
 				return RayHit;
@@ -114,7 +114,7 @@ FInputRayHit UGizmoElementRectangle::LineTrace(const UGizmoViewContext* ViewCont
 				(HdS <= HitBuffer || HdS >= WorldWidth - HitBuffer ||
 				 HdU <= HitBuffer || HdU >= WorldHeight - HitBuffer))
 			{
-				FInputRayHit RayHit(HitDepth);
+				FInputRayHit RayHit(static_cast<float>(HitDepth));
 				RayHit.SetHitObject(this);
 				RayHit.HitIdentifier = PartIdentifier;
 				return RayHit;
diff --git a/Engine/Source/Runtime/InteractiveToolsFramework/Private/BaseGizmos/GizmoElementTorus.cpp b/Engine/Source/Runtime/InteractiveToolsFramework/Private/BaseGizmos/GizmoElementTorus.cpp
--- a/Engine/Source/Runtime/InteractiveToolsFramework/Private/BaseGizmos/GizmoElementTorus.cpp
+++ b/Engine/Source/Runtime/InteractiveToolsFramework/Private/BaseGizmos/GizmoElementTorus.cpp
@@ -27,8 +27,8 @@ void UGizmoElementTorus::Render(IToolsContextRenderAPI* RenderAPI, const FRender
 			bool bPartial = IsPartial(RenderAPI->GetSceneView(), WorldCenter, WorldNormal);
 			FVector BeginAxis = bPartial ? Axis0.RotateAngleAxis(PartialStartAngle, Normal).GetSafeNormal() : Axis0;
 
-			const double PartialAngle = PartialEndAngle - PartialStartAngle;			
-			if (PartialAngle <= 0.0)
+			const float PartialAngle = static_cast<float>(PartialEndAngle - PartialStartAngle);			
+			if (PartialAngle <= 0.0f)
 			{
 				return;
 			}
@@ -91,9 +91,9 @@ FInputRayHit UGizmoElementTorus::LineTrace(const UGizmoViewContext* ViewContext,
 		{
 			// If the torus lies at a glancing angle, the ray-torus intersection is performed against cylinders approximating
 			// the shape of the torus.
-			static constexpr int NumFullTorusCylinders = 16;
+			static constexpr int32 NumFullTorusCylinders = 16;
 			static constexpr double AngleDelta = UE_DOUBLE_TWO_PI / static_cast<double>(NumFullTorusCylinders);
-			const int NumCylinders = bPartial ? FMath::CeilToInt(NumFullTorusCylinders * PartialAngle / UE_DOUBLE_TWO_PI) : NumFullTorusCylinders;
+			const int32 NumCylinders = bPartial ? FMath::CeilToInt32(NumFullTorusCylinders * PartialAngle / UE_DOUBLE_TWO_PI) : NumFullTorusCylinders;
 
 			const FVector ViewDirection = ViewContext->GetViewDirection();
 			FVector VectorA = WorldBeginAxis; 
@@ -108,7 +108,7 @@ FInputRayHit UGizmoElementTorus::LineTrace(const UGizmoViewContext* ViewContext,
 			}
 
 			// Line trace against a set of cylinders approximating the shape of the torus
-			for (int i = 0; i < NumCylinders; i++)
+			for (int32 i = 0; i < NumCylinders; i++)
 			{
 				if (i > 0)
 				{
@@ -155,7 +155,7 @@ FInputRayHit UGizmoElementTorus::LineTrace(const UGizmoViewContext* ViewContext,
 
 			// Find the closest point on the circle to the intersection point
 			FVector NearestCirclePos;
-			GizmoMath::ClosetPointOnCircle(HitPoint, WorldCenter, WorldNormal, WorldOuterRadius, NearestCirclePos);
+			GizmoMath::ClosetPointOnCircle(HitPoint, WorldCenter, WorldNormal, static_cast<float>(WorldOuterRadius), NearestCirclePos);
 
 			// Find the closest point on the ray to the circle and determine if it is within the torus
 			FRay Ray(RayOrigin, RayDirection, true);
@@ -193,7 +193,7 @@ FInputRayHit UGizmoElementTorus::LineTrace(const UGizmoViewContext* ViewContext,
 
 		if (HitDepth >= 0.0)
 		{
-			FInputRayHit RayHit(HitDepth);
+			FInputRayHit RayHit(static_cast<float>(HitDepth));
 			RayHit.SetHitObject(this);
 			RayHit.HitIdentifier = PartIdentifier;
 			return RayHit;
diff --git a/Engine/Source/Runtime/InteractiveToolsFramework/Private/BaseGizmos/GizmoLineHandleComponent.cpp b/Engine/Source/Runtime/InteractiveToolsFramework/Private/BaseGizmos/GizmoLineHandleComponent.cpp
--- a/Engine/Source/Runtime/InteractiveToolsFramework/Private/BaseGizmos/GizmoLineHandleComponent.cpp
+++ b/Engine/Source/Runtime/InteractiveToolsFramework/Private/BaseGizmos/GizmoLineHandleComponent.cpp
@@ -73,13 +73,13 @@ public:
 					(HoverThicknessMultiplier * Thickness) : (Thickness);
 				if (!bIsOrtho)
 				{
-					UseThickness *= (View->FOV / 90.0);		// compensate for FOV scaling in Gizmos...
+					UseThickness *= (View->FOV / 90.0f);		// compensate for FOV scaling in Gizmos...
 				}
 
 				// From base origin to disk origin
 				PDI->DrawLine(WorldBaseOrigin, ScaledDiskOrigin, Color, SDPG_Foreground, UseThickness, 0.0f, true);
 				// Draw the interval marker
-				PDI->DrawLine(ScaledIntevalStart, ScaledIntevalEnd, Color, SDPG_Foreground, 2. * UseThickness, 0.0f, true);
+				PDI->DrawLine(ScaledIntevalStart, ScaledIntevalEnd, Color, SDPG_Foreground, 2 * UseThickness, 0.0f, true);
 			}
 		}
 	}
@@ -100,8 +100,8 @@ public:
 		return false;
 	}
 
-	virtual uint32 GetMemoryFootprint(void) const override { return sizeof *this + GetAllocatedSize(); }
-	uint32 GetAllocatedSize(void) const { return FPrimitiveSceneProxy::GetAllocatedSize(); }
+	virtual uint32 GetMemoryFootprint(void) const override { return IntCastChecked<uint32>( sizeof *this + GetAllocatedSize() ); }
+	SIZE_T GetAllocatedSize(void) const { return FPrimitiveSceneProxy::GetAllocatedSize(); }
 
 
 	void SetExternalHoverState(bool* HoverState)
@@ -197,7 +197,7 @@ bool UGizmoLineHandleComponent::LineTraceComponent(FHitResult& OutHit, const FVe
 	}
 
 	OutHit.Component = this;
-	OutHit.Distance = FVector::Distance(Start, NearestOnLine);
+	OutHit.Distance = static_cast<float>(FVector::Distance(Start, NearestOnLine));
 	OutHit.ImpactPoint = NearestOnLine;
 	return true;
 
diff --git a/Engine/Source/Runtime/InteractiveToolsFramework/Private/BaseGizmos/GizmoMath.cpp b/Engine/Source/Runtime/InteractiveToolsFramework/Private/BaseGizmos/GizmoMath.cpp
--- a/Engine/Source/Runtime/InteractiveToolsFramework/Private/BaseGizmos/GizmoMath.cpp
+++ b/Engine/Source/Runtime/InteractiveToolsFramework/Private/BaseGizmos/GizmoMath.cpp
@@ -162,7 +162,7 @@ void GizmoMath::RayCylinderIntersection(
 	Basis[1] = GetOrthogonalVector(Basis[0]).GetSafeNormal();
 	Basis[2] = (Basis[0] ^ Basis[1]).GetSafeNormal();
 
-	RealType HalfHeight = 0.5 * CylinderHeight;
+	RealType HalfHeight = RealType(0.5) * CylinderHeight;
 	RealType RadiusSquared = CylinderRadius * CylinderRadius;
 
 	// Convert incoming line origin to capsule coordinates.
@@ -171,25 +171,25 @@ void GizmoMath::RayCylinderIntersection(
 
 	// Get the z-value, in cylinder coordinates, of the incoming line's
 	// unit-length direction.
-	RealType Dz = Basis[0] | RayDirection;
+	RealType Dz = static_cast<RealType>( Basis[0] | RayDirection );
 	if (FMath::IsNearlyEqual(Dz, 1.0, SMALL_NUMBER)) 
 	{
 		// The line is parallel to the cylinder axis.  Determine whether the
 		// line intersects the cylinder end disks.
-		RealType RadialSquaredDist = RadiusSquared - P[0] * P[0] - P[1] * P[1];
+		RealType RadialSquaredDist = static_cast<RealType>( RadiusSquared - P[0] * P[0] - P[1] * P[1]);
 		if (RadialSquaredDist >= 0.0)
 		{
 			// The line intersects the cylinder end disks.
 			NumIntersections = 2;
 			if (Dz > 0.0)
 			{
-				RayParam[0] = -P[2] - HalfHeight;
-				RayParam[1] = -P[2] + HalfHeight;
+				RayParam[0] = static_cast<RealType>( -P[2] - HalfHeight );
+				RayParam[1] = static_cast<RealType>( -P[2] + HalfHeight );
 			}
 			else
 			{
-				RayParam[0] = P[2] - HalfHeight;
-				RayParam[1] = P[2] + HalfHeight;
+				RayParam[0] = static_cast<RealType>( P[2] - HalfHeight );
+				RayParam[1] = static_cast<RealType>( P[2] + HalfHeight );
 			}
 		}
 		// else:  The line is outside the cylinder, no intersection.
@@ -212,9 +212,9 @@ void GizmoMath::RayCylinderIntersection(
 				// quadratic equation.  If P = (px,py,pz) and D = (dx,dy,dz),
 				// then the quadratic equation is
 				//   (dx^2+dy^2)*t^2 + 2*(px*dx+py*dy)*t + (px^2+py^2-r^2) = 0
-				A0 = P[0] * P[0] + P[1] * P[1] - RadiusSquared;
-				A1 = P[0] * D[0] + P[1] * D[1];
-				A2 = D[0] * D[0] + D[1] * D[1];
+				A0 = static_cast<RealType>( P[0] * P[0] + P[1] * P[1] - RadiusSquared );
+				A1 = static_cast<RealType>( P[0] * D[0] + P[1] * D[1] );
+				A2 = static_cast<RealType>( D[0] * D[0] + D[1] * D[1] );
 				Discr = A1 * A1 - A0 * A2;
 				if (FMath::IsNearlyZero(Discr)) 
 				{
@@ -239,20 +239,20 @@ void GizmoMath::RayCylinderIntersection(
 		else
 		{
 			// Test for intersections with the planes of the end disks.
-			Inv = 1.0f / D[2];
+			Inv = static_cast<RealType>( 1.0 / D[2] );
 
-			RealType T0 = (-HalfHeight - P[2]) * Inv;
-			RealType TmpX = P[0] + T0 * D[0];
-			RealType TmpY = P[1] + T0 * D[1];
+			RealType T0 = static_cast<RealType>( (-HalfHeight - P[2]) * Inv );
+			RealType TmpX = static_cast<RealType>( P[0] + T0 * D[0] );
+			RealType TmpY = static_cast<RealType>( P[1] + T0 * D[1] );
 			if (TmpX * TmpX + TmpY * TmpY <= RadiusSquared)
 			{
 				// Plane intersection inside the top cylinder end disk.
 				RayParam[NumIntersections++] = T0;
 			}
 
-			RealType T1 = (+HalfHeight - P[2]) * Inv;
-			TmpX = P[0] + T1 * D[0];
-			TmpY = P[1] + T1 * D[1];
+			RealType T1 = static_cast<RealType>( (+HalfHeight - P[2]) * Inv );
+			TmpX = static_cast<RealType>( P[0] + T1 * D[0] );
+			TmpY = static_cast<RealType>( P[1] + T1 * D[1] );
 			if (TmpX * TmpX + TmpY * TmpY <= RadiusSquared)
 			{
 				// Plane intersection inside the bottom cylinder end disk.
@@ -262,9 +262,9 @@ void GizmoMath::RayCylinderIntersection(
 			if (NumIntersections < 2)
 			{
 				// Test for intersection with the cylinder wall.
-				A0 = P[0] * P[0] + P[1] * P[1] - RadiusSquared;
-				A1 = P[0] * D[0] + P[1] * D[1];
-				A2 = D[0] * D[0] + D[1] * D[1];
+				A0 = static_cast<RealType>( P[0] * P[0] + P[1] * P[1] - RadiusSquared );
+				A1 = static_cast<RealType>( P[0] * D[0] + P[1] * D[1] );
+				A2 = static_cast<RealType>( D[0] * D[0] + D[1] * D[1] );
 				Discr = A1 * A1 - A0 * A2;
 				if (FMath::IsNearlyZero(Discr))
 				{
@@ -395,10 +395,10 @@ void GizmoMath::RayConeIntersection(
 	// equation subject to the linear inequality constraints.
 
 	FVector PmV = RayOrigin - ConeCenter;
-	RealType DdU = FVector::DotProduct(ConeDirection, RayDirection);
-	RealType DdPmV = FVector::DotProduct(ConeDirection, PmV);
-	RealType UdPmV = FVector::DotProduct(RayDirection, PmV);
-	RealType PmVdPmV = FVector::DotProduct(PmV, PmV);
+	RealType DdU = static_cast<RealType>( FVector::DotProduct(ConeDirection, RayDirection) );
+	RealType DdPmV = static_cast<RealType>( FVector::DotProduct(ConeDirection, PmV) );
+	RealType UdPmV = static_cast<RealType>( FVector::DotProduct(RayDirection, PmV) );
+	RealType PmVdPmV = static_cast<RealType>( FVector::DotProduct(PmV, PmV) );
 	RealType CosAngleSqr = ConeCosAngle * ConeCosAngle;
 	RealType C2 = DdU * DdU - CosAngleSqr;
 	RealType C1 = DdU * DdPmV - CosAngleSqr * UdPmV;
@@ -439,7 +439,7 @@ void GizmoMath::RayConeIntersection(
 			// or both of them might intersect the negative cone.  We are
 			// interested only in those intersections with the positive cone.
 			RealType Root = FMath::Sqrt(Discr);
-			RealType InvC2 = 1.0 / C2;
+			RealType InvC2 = RealType(1) / C2;
 			int NumIntersections = 0;
 
 			T = (-C1 - Root) * InvC2;
diff --git a/Engine/Source/Runtime/InteractiveToolsFramework/Private/BaseGizmos/GizmoRectangleComponent.cpp b/Engine/Source/Runtime/InteractiveToolsFramework/Private/BaseGizmos/GizmoRectangleComponent.cpp
--- a/Engine/Source/Runtime/InteractiveToolsFramework/Private/BaseGizmos/GizmoRectangleComponent.cpp
+++ b/Engine/Source/Runtime/InteractiveToolsFramework/Private/BaseGizmos/GizmoRectangleComponent.cpp
@@ -133,7 +133,7 @@ public:
 					(HoverThicknessMultiplier * Thickness) : (Thickness);
 				if (View->IsPerspectiveProjection())
 				{
-					UseThickness *= (View->FOV / 90.0);		// compensate for FOV scaling in Gizmos...
+					UseThickness *= (View->FOV / 90.0f);		// compensate for FOV scaling in Gizmos...
 				}
 
 				if (SegmentFlags & 0x1)
@@ -192,8 +192,8 @@ public:
 		return false;
 	}
 
-	virtual uint32 GetMemoryFootprint(void) const override { return sizeof *this + GetAllocatedSize(); }
-	uint32 GetAllocatedSize(void) const { return FPrimitiveSceneProxy::GetAllocatedSize(); }
+	virtual uint32 GetMemoryFootprint(void) const override { return IntCastChecked<uint32>( sizeof *this + GetAllocatedSize() ); }
+	SIZE_T GetAllocatedSize(void) const { return FPrimitiveSceneProxy::GetAllocatedSize(); }
 
 	void SetExternalHoverState(bool* HoverState)
 	{
@@ -274,7 +274,7 @@ bool UGizmoRectangleComponent::LineTraceComponent(FHitResult& OutHit, const FVec
 			HitPoint, HitNormal) )
 		{
 			OutHit.Component = this;
-			OutHit.Distance = FVector::Distance(Start, HitPoint);
+			OutHit.Distance = static_cast<float>( FVector::Distance(Start, HitPoint) );
 			OutHit.ImpactPoint = HitPoint;
 			OutHit.ImpactNormal = HitNormal;
 			return true;
diff --git a/Engine/Source/Runtime/InteractiveToolsFramework/Private/BaseGizmos/GizmoRenderingUtil.cpp b/Engine/Source/Runtime/InteractiveToolsFramework/Private/BaseGizmos/GizmoRenderingUtil.cpp
--- a/Engine/Source/Runtime/InteractiveToolsFramework/Private/BaseGizmos/GizmoRenderingUtil.cpp
+++ b/Engine/Source/Runtime/InteractiveToolsFramework/Private/BaseGizmos/GizmoRenderingUtil.cpp
@@ -69,7 +69,7 @@ namespace GizmoRenderingUtilLocals
 		// below gives us a sense of this distance, so we make the offset relative to that
 		// (this does do one redundant WorldToScreen)
 		FVector4 LocationScreenPoint = View->WorldToScreen(Location);
-		float OffsetDelta = LocationScreenPoint.W * 0.01f;
+		double OffsetDelta = LocationScreenPoint.W * 0.01;
 
 		FVector2D PixelA = WorldToPixelDouble(View, Location);
 		FVector OffsetPointWorld = Location + OffsetDelta * View->GetViewRight() + OffsetDelta * View->GetViewUp();
@@ -106,7 +106,7 @@ float GizmoRenderingUtil::CalculateViewDependentScaleAndFlatten(
 	bool bIsOrthoXY = !bIsPerspective && FMath::Abs(ViewMatrix.M[2][2]) > 0.0f;
 	bool bIsOrthoXZ = !bIsPerspective && FMath::Abs(ViewMatrix.M[1][2]) > 0.0f;
 	bool bIsOrthoYZ = !bIsPerspective && FMath::Abs(ViewMatrix.M[0][2]) > 0.0f;
-	float UniformScale = InScale * View->WorldToScreen(Location).W * (4.0f / View->UnscaledViewRect.Width() / View->ViewMatrices.GetProjectionMatrix().M[0][0]);
+	float UniformScale = static_cast<float> (InScale * View->WorldToScreen(Location).W * (4.0 / View->UnscaledViewRect.Width() / View->ViewMatrices.GetProjectionMatrix().M[0][0]));
 
 	// Clamp to tolerance to prevent division by zero.
 	// @todo change to use MathUtil<RealType>::ZeroTolerance and TMathUtil<RealType>::SignNonZero(Value)
diff --git a/Engine/Source/Runtime/InteractiveToolsFramework/Private/BaseGizmos/PlanePositionGizmo.cpp b/Engine/Source/Runtime/InteractiveToolsFramework/Private/BaseGizmos/PlanePositionGizmo.cpp
--- a/Engine/Source/Runtime/InteractiveToolsFramework/Private/BaseGizmos/PlanePositionGizmo.cpp
+++ b/Engine/Source/Runtime/InteractiveToolsFramework/Private/BaseGizmos/PlanePositionGizmo.cpp
@@ -84,12 +84,12 @@ void UPlanePositionGizmo::OnClickPress(const FInputDeviceRay& PressPos)
 
 	FVector AxisOrigin = AxisSource->GetOrigin();
 
-	float DirectionSignX = FVector::DotProduct(InteractionStartPoint - AxisOrigin, InteractionAxisX);
-	ParameterSigns.X = (bEnableSignedAxis && DirectionSignX < 0) ? -1.0f : 1.0f;
-	ParameterSigns.X *= (bFlipX) ? -1.0f : 1.0f;
-	float DirectionSignY = FVector::DotProduct(InteractionStartPoint - AxisOrigin, InteractionAxisY);
-	ParameterSigns.Y = (bEnableSignedAxis && DirectionSignY < 0) ? -1.0f : 1.0f;
-	ParameterSigns.Y *= (bFlipY) ? -1.0f : 1.0f;
+	double DirectionSignX = FVector::DotProduct(InteractionStartPoint - AxisOrigin, InteractionAxisX);
+	ParameterSigns.X = (bEnableSignedAxis && DirectionSignX < 0) ? -1.0 : 1.0;
+	ParameterSigns.X *= (bFlipX) ? -1.0 : 1.0;
+	double DirectionSignY = FVector::DotProduct(InteractionStartPoint - AxisOrigin, InteractionAxisY);
+	ParameterSigns.Y = (bEnableSignedAxis && DirectionSignY < 0) ? -1.0 : 1.0;
+	ParameterSigns.Y *= (bFlipY) ? -1.0 : 1.0;
 
 	InteractionStartParameter = GizmoMath::ComputeCoordinatesInPlane(IntersectionPoint,
 			InteractionOrigin, InteractionNormal, InteractionAxisX, InteractionAxisY);
diff --git a/Engine/Source/Runtime/InteractiveToolsFramework/Private/BaseTools/BaseBrushTool.cpp b/Engine/Source/Runtime/InteractiveToolsFramework/Private/BaseTools/BaseBrushTool.cpp
--- a/Engine/Source/Runtime/InteractiveToolsFramework/Private/BaseTools/BaseBrushTool.cpp
+++ b/Engine/Source/Runtime/InteractiveToolsFramework/Private/BaseTools/BaseBrushTool.cpp
@@ -33,8 +33,8 @@ void UBaseBrushTool::Setup()
 {
 	UMeshSurfacePointTool::Setup();
 	BrushProperties = NewObject<UBrushBaseProperties>(this, PropertyClass.Get(), TEXT("Brush"));
-	double MaxDimension = EstimateMaximumTargetDimension();
-	BrushRelativeSizeRange = TInterval<float>(MaxDimension*0.01, MaxDimension);
+	float MaxDimension = static_cast<float>( EstimateMaximumTargetDimension());
+	BrushRelativeSizeRange = TInterval<float>(MaxDimension*0.01f, MaxDimension);
 
 	RecalculateBrushRadius();
 
@@ -66,7 +66,7 @@ void UBaseBrushTool::IncreaseBrushSizeAction()
 	if (BrushProperties->bSpecifyRadius)
 	{
 		// Hardcoded max of 1000 chosen to match the BrushRadius "UIMax" specified in UBrushBaseProperties
-		BrushProperties->BrushRadius = FMath::Min(BrushProperties->BrushRadius * 1.1, 1000);
+		BrushProperties->BrushRadius = FMath::Min(BrushProperties->BrushRadius * 1.1f, 1000.f);
 	}
 	else
 	{
@@ -80,7 +80,7 @@ void UBaseBrushTool::DecreaseBrushSizeAction()
 {
 	if (BrushProperties->bSpecifyRadius)
 	{
-		BrushProperties->BrushRadius = FMath::Max(BrushProperties->BrushRadius / 1.1, 1);
+		BrushProperties->BrushRadius = FMath::Max(BrushProperties->BrushRadius / 1.1f, 1.f);
 	}
 	else
 	{
@@ -187,12 +187,12 @@ void UBaseBrushTool::RecalculateBrushRadius()
 	if (BrushProperties->bSpecifyRadius)
 	{
 		CurrentBrushRadius = BrushProperties->BrushRadius;
-		BrushProperties->BrushSize = (2 * CurrentBrushRadius - ScaledBrushSizeRange.Min) / ScaledBrushSizeRange.Size();
+		BrushProperties->BrushSize = static_cast<float>( (2 * CurrentBrushRadius - ScaledBrushSizeRange.Min) / ScaledBrushSizeRange.Size() );
 	}
 	else
 	{
 		CurrentBrushRadius = 0.5 * ScaledBrushSizeRange.Interpolate(BrushProperties->BrushSize);
-		BrushProperties->BrushRadius = CurrentBrushRadius;
+		BrushProperties->BrushRadius = static_cast<float>( CurrentBrushRadius );
 	}
 }
 
diff --git a/Engine/Source/Runtime/InteractiveToolsFramework/Private/SceneQueries/SceneSnappingManager.cpp b/Engine/Source/Runtime/InteractiveToolsFramework/Private/SceneQueries/SceneSnappingManager.cpp
--- a/Engine/Source/Runtime/InteractiveToolsFramework/Private/SceneQueries/SceneSnappingManager.cpp
+++ b/Engine/Source/Runtime/InteractiveToolsFramework/Private/SceneQueries/SceneSnappingManager.cpp
@@ -85,7 +85,7 @@ void FSceneHitQueryResult::InitializeHitResult(const FSceneHitQueryRequest& From
 {
 	HitResult = FHitResult(TargetActor, TargetComponent, (FVector)Position, (FVector)Normal);
 
-	HitResult.Distance = FromRequest.WorldRay.GetParameter(Position);
+	HitResult.Distance = static_cast<float>( FromRequest.WorldRay.GetParameter(Position) );
 	HitResult.FaceIndex = HitTriIndex;
 	// initialize .Time? Need start/end for that...
 }
diff --git a/Engine/Source/Runtime/InteractiveToolsFramework/Private/ToolBuilderUtil.cpp b/Engine/Source/Runtime/InteractiveToolsFramework/Private/ToolBuilderUtil.cpp
--- a/Engine/Source/Runtime/InteractiveToolsFramework/Private/ToolBuilderUtil.cpp
+++ b/Engine/Source/Runtime/InteractiveToolsFramework/Private/ToolBuilderUtil.cpp
@@ -13,7 +13,7 @@ int ToolBuilderUtil::CountComponents(const FToolBuilderState& InputState, const
 
 	if (InputState.SelectedComponents.Num() > 0)
 	{
-		nTypedComponents = Algo::CountIf(InputState.SelectedComponents, Predicate);
+		nTypedComponents = static_cast<int>( Algo::CountIf(InputState.SelectedComponents, Predicate) );
 	}
 	else
 	{
@@ -21,7 +21,7 @@ int ToolBuilderUtil::CountComponents(const FToolBuilderState& InputState, const
 			Algo::TransformAccumulate(InputState.SelectedActors,
 									  [&Predicate](AActor* Actor)
 									  {
-										  return Algo::CountIf(Actor->GetComponents(), Predicate);
+										  return static_cast<int>( Algo::CountIf(Actor->GetComponents(), Predicate));
 									  },
 									  0);
 	}
@@ -109,7 +109,7 @@ void ToolBuilderUtil::EnumerateComponents(const FToolBuilderState& InputState, T
 
 int32 ToolBuilderUtil::CountActors(const FToolBuilderState& InputState, const TFunction<bool(AActor*)>& Predicate)
 {
-	return Algo::CountIf(InputState.SelectedActors, Predicate);
+	return static_cast<int32>( Algo::CountIf(InputState.SelectedActors, Predicate) );
 }
 
 
diff --git a/Engine/Source/Runtime/InteractiveToolsFramework/Private/ToolDataVisualizer.cpp b/Engine/Source/Runtime/InteractiveToolsFramework/Private/ToolDataVisualizer.cpp
--- a/Engine/Source/Runtime/InteractiveToolsFramework/Private/ToolDataVisualizer.cpp
+++ b/Engine/Source/Runtime/InteractiveToolsFramework/Private/ToolDataVisualizer.cpp
@@ -86,7 +86,7 @@ void FToolDataVisualizer::PopAllTransforms()
 void FToolDataVisualizer::InternalDrawTransformedLine(const FVector& A, const FVector& B, const FLinearColor& ColorIn, float LineThicknessIn, bool bDepthTestedIn)
 {
 	CurrentPDI->DrawLine(A, B, ColorIn,
-		(bDepthTestedIn) ? SDPG_World : SDPG_Foreground,
+		uint8( (bDepthTestedIn) ? SDPG_World : SDPG_Foreground),
 		LineThicknessIn * PDISizeScale, DepthBias, true);
 }
 
@@ -94,7 +94,7 @@ void FToolDataVisualizer::InternalDrawTransformedLine(const FVector& A, const FV
 void FToolDataVisualizer::InternalDrawTransformedPoint(const FVector& Position, const FLinearColor& ColorIn, float PointSizeIn, bool bDepthTestedIn)
 {
 	CurrentPDI->DrawPoint(Position, ColorIn, PointSizeIn * PDISizeScale,
-		(bDepthTestedIn) ? SDPG_World : SDPG_Foreground);
+		uint8( (bDepthTestedIn) ? SDPG_World : SDPG_Foreground) );
 }
 
 
@@ -107,7 +107,7 @@ void FToolDataVisualizer::InternalDrawCircle(const FVector& Position, const FVec
 	// this function is from SceneManagement.h
 	::DrawCircle(CurrentPDI, TransformP(Position), (FVector)Tan1, (FVector)Tan2, 
 		Color, Radius, Steps,
-		(bDepthTestedIn) ? SDPG_World : SDPG_Foreground,
+		uint8( (bDepthTestedIn) ? SDPG_World : SDPG_Foreground),
 		LineThicknessIn * PDISizeScale, DepthBias, true);
 }
 
@@ -191,7 +191,7 @@ void FToolDataVisualizer::InternalDrawViewFacingCircle(const FVector& Position,
 	// this function is from SceneManagement.h
 	::DrawCircle(CurrentPDI, WorldPosition, (FVector)Tan1, (FVector)Tan2,
 		Color, Radius, Steps,
-		(bDepthTestedIn) ? SDPG_World : SDPG_Foreground,
+		uint8( (bDepthTestedIn) ? SDPG_World : SDPG_Foreground),
 		LineThicknessIn * PDISizeScale, DepthBias, true);
 }
 
diff --git a/Engine/Source/Runtime/InteractiveToolsFramework/Public/BaseGizmos/GizmoElementGroup.h b/Engine/Source/Runtime/InteractiveToolsFramework/Public/BaseGizmos/GizmoElementGroup.h
--- a/Engine/Source/Runtime/InteractiveToolsFramework/Public/BaseGizmos/GizmoElementGroup.h
+++ b/Engine/Source/Runtime/InteractiveToolsFramework/Public/BaseGizmos/GizmoElementGroup.h
@@ -72,5 +72,5 @@ protected:
 	TArray<TObjectPtr<UGizmoElementBase>> Elements;
 
 	// Updates input transform's scale component to have uniform scale and applies constant scale if bConstantScale is true
-	virtual void ApplyUniformConstantScaleToTransform(float PixelToWorldScale, FTransform& InOutLocalToWorldTransform) const;
+	virtual void ApplyUniformConstantScaleToTransform(double PixelToWorldScale, FTransform& InOutLocalToWorldTransform) const;
 };
diff --git a/Engine/Source/Runtime/InteractiveToolsFramework/Public/BaseGizmos/GizmoViewContext.h b/Engine/Source/Runtime/InteractiveToolsFramework/Public/BaseGizmos/GizmoViewContext.h
--- a/Engine/Source/Runtime/InteractiveToolsFramework/Public/BaseGizmos/GizmoViewContext.h
+++ b/Engine/Source/Runtime/InteractiveToolsFramework/Public/BaseGizmos/GizmoViewContext.h
@@ -71,11 +71,11 @@ public:
 		if (ScreenPoint.W != 0.0f)
 		{
 			//Reverse the W in the case it is negative, this allow to manipulate a manipulator in the same direction when the camera is really close to the manipulator.
-			float InvW = (ScreenPoint.W > 0.0f ? 1.0f : -1.0f) / ScreenPoint.W;
-			float Y = (GProjectionSignY > 0.0f) ? ScreenPoint.Y : 1.0f - ScreenPoint.Y;
+			double InvW = (ScreenPoint.W > 0.0 ? 1.0 : -1.0) / ScreenPoint.W;
+			double Y = (GProjectionSignY > 0.0) ? ScreenPoint.Y : 1.0 - ScreenPoint.Y;
 			OutPixelLocation = FVector2D(
-				UnscaledViewRect.Min.X + (0.5f + ScreenPoint.X * 0.5f * InvW) * UnscaledViewRect.Width(),
-				UnscaledViewRect.Min.Y + (0.5f - Y * 0.5f * InvW) * UnscaledViewRect.Height()
+				UnscaledViewRect.Min.X + (0.5 + ScreenPoint.X * 0.5 * InvW) * UnscaledViewRect.Width(),
+				UnscaledViewRect.Min.Y + (0.5 - Y * 0.5 * InvW) * UnscaledViewRect.Height()
 			);
 			return true;
 		}
diff --git a/Engine/Source/Runtime/InteractiveToolsFramework/Public/ToolDataVisualizer.h b/Engine/Source/Runtime/InteractiveToolsFramework/Public/ToolDataVisualizer.h
--- a/Engine/Source/Runtime/InteractiveToolsFramework/Public/ToolDataVisualizer.h
+++ b/Engine/Source/Runtime/InteractiveToolsFramework/Public/ToolDataVisualizer.h
@@ -82,7 +82,7 @@ public:
 	{
 		// transform normal by a safe inverse scale + normalize, and a standard rotation (TODO: move implementation into FTransform)
 		const FVector& S = TotalTransform.GetScale3D();
-		float DetSign = FMath::Sign(S.X*S.Y*S.Z); // we only need to multiply by the sign of the determinant, rather than divide by it, since we normalize later anyway
+		double DetSign = FMath::Sign(S.X*S.Y*S.Z); // we only need to multiply by the sign of the determinant, rather than divide by it, since we normalize later anyway
 		FVector SafeInvS(S.Y*S.Z*DetSign, S.X*S.Z*DetSign, S.X*S.Y*DetSign);
 		return TotalTransform.TransformVectorNoScale((SafeInvS*Normal).GetSafeNormal());
 	}
