ITF : ScalableSphereGizmo - removing unused kismet include.

diff --git a/Engine/Source/Runtime/InteractiveToolsFramework/Private/BaseGizmos/ScalableSphereGizmo.cpp b/Engine/Source/Runtime/InteractiveToolsFramework/Private/BaseGizmos/ScalableSphereGizmo.cpp
--- a/Engine/Source/Runtime/InteractiveToolsFramework/Private/BaseGizmos/ScalableSphereGizmo.cpp
+++ b/Engine/Source/Runtime/InteractiveToolsFramework/Private/BaseGizmos/ScalableSphereGizmo.cpp
@@ -6,7 +6,6 @@
 #include "BaseGizmos/GizmoBoxComponent.h"
 #include "Components/SphereComponent.h"
 #include "Engine/CollisionProfile.h"
-#include "Kismet/KismetMathLibrary.h"
 #include "BaseGizmos/GizmoMath.h"
 #include "BaseBehaviors/MouseHoverBehavior.h"
 #include "BaseGizmos/GizmoRenderingUtil.h"
