InteractiveToolsContext: add EToolContextTransformGizmoMode enum and IToolsContextQueriesAPI::GetCurrentTransformGizmoMode(), this provides a mechanism for ITF tools/gizmos to query the active Gizmo mode in the parent Editor.
Add UEditorInteractiveToolsContext::SetForceCombinedGizmoMode() and GetForceCombinedGizmoModeEnabled() so that an ITC can have a toggle for whether or not to respect the Editor gizmo mode.
Add implementation of new ::GetCurrentTransformGizmoMode() in FEdModeToolsContextQueriesImpl, that uses this new flag.
Add support in UCombinedTransformGizmo for specifying which sub-gizmo set is visible based on EToolContextTransformGizmoMode member, or alternately querying EToolContextTransformGizmoMode via IToolsContextQueriesAPI.

ModelingMode: add new UModelingToolsEditorModeSettings::bRespectLevelEditorGizmoMode to configure combined-vs-leveleditor gizmo mode behavior of Modeling Mode. Currently applying this setting to ModelingMode ITC in ::Begin(), so mode exit/enter is needed to apply the setting.

diff --git a/Engine/Plugins/Editor/ModelingToolsEditorMode/Source/ModelingToolsEditorMode/Private/ModelingToolsEditorMode.cpp b/Engine/Plugins/Editor/ModelingToolsEditorMode/Source/ModelingToolsEditorMode/Private/ModelingToolsEditorMode.cpp
--- a/Engine/Plugins/Editor/ModelingToolsEditorMode/Source/ModelingToolsEditorMode/Private/ModelingToolsEditorMode.cpp
+++ b/Engine/Plugins/Editor/ModelingToolsEditorMode/Source/ModelingToolsEditorMode/Private/ModelingToolsEditorMode.cpp
@@ -339,6 +339,8 @@ void UModelingToolsEditorMode::Enter()
 {
 	UEdMode::Enter();
 
+	const UModelingToolsEditorModeSettings* ModelingModeSettings = GetDefault<UModelingToolsEditorModeSettings>();
+
 	// Register builders for tool targets that the mode uses.
 	GetInteractiveToolsContext()->TargetManager->AddTargetFactory(NewObject<UStaticMeshComponentToolTargetFactory>(GetToolManager()));
 	GetInteractiveToolsContext()->TargetManager->AddTargetFactory(NewObject<UVolumeComponentToolTargetFactory>(GetToolManager()));
@@ -363,13 +365,17 @@ void UModelingToolsEditorMode::Enter()
 
 	// register gizmo helper
 	UE::TransformGizmoUtil::RegisterTransformGizmoContextObject(GetInteractiveToolsContext());
+	// configure mode-level Gizmo options
+	if (ModelingModeSettings)
+	{
+		GetInteractiveToolsContext()->SetForceCombinedGizmoMode(ModelingModeSettings->bRespectLevelEditorGizmoMode == false);
+	}
 
 	// register snapping manager
 	UE::Geometry::RegisterSceneSnappingManager(GetInteractiveToolsContext());
 	SceneSnappingManager = UE::Geometry::FindModelingSceneSnappingManager(GetToolManager());
 
 	// register selection manager, if this feature is enabled in the mode settings
-	const UModelingToolsEditorModeSettings* ModelingModeSettings = GetDefault<UModelingToolsEditorModeSettings>();
 	if (ModelingModeSettings && ModelingModeSettings->bEnablePersistentSelections)
 	{
 		// set up SelectionManager and register known factory types
diff --git a/Engine/Plugins/Editor/ModelingToolsEditorMode/Source/ModelingToolsEditorMode/Public/ModelingToolsEditorModeSettings.h b/Engine/Plugins/Editor/ModelingToolsEditorMode/Source/ModelingToolsEditorMode/Public/ModelingToolsEditorModeSettings.h
--- a/Engine/Plugins/Editor/ModelingToolsEditorMode/Source/ModelingToolsEditorMode/Public/ModelingToolsEditorModeSettings.h
+++ b/Engine/Plugins/Editor/ModelingToolsEditorMode/Source/ModelingToolsEditorMode/Public/ModelingToolsEditorModeSettings.h
@@ -133,6 +133,14 @@ public:
 	UPROPERTY(config, EditAnywhere, Category = "Modeling Mode|Generated Assets")
 	bool bAppendRandomStringToName = true;
 
+	//
+	// User Interface
+	//
+
+	/** If true, the standard UE Editor Gizmo Mode (ie selected via the Level Editor Viewport toggle) will be used to configure the Modeling Gizmo, otherwise a Combined Gizmo will always be used. It may be necessary to exit and re-enter Modeling Mode after changing this setting. */
+	UPROPERTY(config, EditAnywhere, Category = "Modeling Mode|User Interface")
+	bool bRespectLevelEditorGizmoMode = false;
+
 	//
 	// Selection
 	//
diff --git a/Engine/Source/Editor/Experimental/EditorInteractiveToolsFramework/Private/EdModeInteractiveToolsContext.cpp b/Engine/Source/Editor/Experimental/EditorInteractiveToolsFramework/Private/EdModeInteractiveToolsContext.cpp
--- a/Engine/Source/Editor/Experimental/EditorInteractiveToolsFramework/Private/EdModeInteractiveToolsContext.cpp
+++ b/Engine/Source/Editor/Experimental/EditorInteractiveToolsFramework/Private/EdModeInteractiveToolsContext.cpp
@@ -139,6 +139,27 @@ public:
 		return (CoordSys == COORD_World) ? EToolContextCoordinateSystem::World : EToolContextCoordinateSystem::Local;
 	}
 
+	virtual EToolContextTransformGizmoMode GetCurrentTransformGizmoMode() const override
+	{
+		if (ToolsContext->GetForceCombinedGizmoModeEnabled() == false)
+		{
+			UE::Widget::EWidgetMode WidgetMode = EditorModeManager->GetWidgetMode();
+			switch (WidgetMode)
+			{
+			case UE::Widget::EWidgetMode::WM_None:
+				return EToolContextTransformGizmoMode::NoGizmo;
+			case UE::Widget::EWidgetMode::WM_Translate:
+				return EToolContextTransformGizmoMode::Translation;
+			case UE::Widget::EWidgetMode::WM_Rotate:
+				return EToolContextTransformGizmoMode::Rotation;
+			case UE::Widget::EWidgetMode::WM_Scale:
+				return EToolContextTransformGizmoMode::Scale;
+			}
+		}
+		
+		return EToolContextTransformGizmoMode::Combined;
+	}
+
 	virtual FToolContextSnappingConfiguration GetCurrentSnappingSettings() const override
 	{
 		FToolContextSnappingConfiguration Config;
@@ -815,6 +836,11 @@ void UEditorInteractiveToolsContext::SetEnableRenderingDuringHitProxyPass(bool b
 	bEnableRenderingDuringHitProxyPass = bEnabled;
 }
 
+void UEditorInteractiveToolsContext::SetForceCombinedGizmoMode(bool bEnabled)
+{
+	bForceCombinedGizmoMode = bEnabled;
+}
+
 void UEditorInteractiveToolsContext::OnLevelEditorCreated(TSharedPtr<ILevelEditor> InLevelEditor)
 {
 	if (UTypedElementSelectionSet* TypedElementSelectionSet = EditorModeManager->GetEditorSelectionSet())
diff --git a/Engine/Source/Editor/Experimental/EditorInteractiveToolsFramework/Public/EdModeInteractiveToolsContext.h b/Engine/Source/Editor/Experimental/EditorInteractiveToolsFramework/Public/EdModeInteractiveToolsContext.h
--- a/Engine/Source/Editor/Experimental/EditorInteractiveToolsFramework/Public/EdModeInteractiveToolsContext.h
+++ b/Engine/Source/Editor/Experimental/EditorInteractiveToolsFramework/Public/EdModeInteractiveToolsContext.h
@@ -148,6 +148,17 @@ public:
 	/** @return true if HitProxy rendering will be allowed in ::Render() */
 	bool GetEnableRenderingDuringHitProxyPass() const { return bEnableRenderingDuringHitProxyPass; }
 
+
+	/**
+	 * Configure whether Transform Gizmos created by the ITF (eg CombinedTransformGizmo) should prefer to show in 'Combined' mode.
+	 * If this is disabled, the Gizmo should respect the active Editor Gizmo setting (eg in the Level Viewport)
+	 */
+	void SetForceCombinedGizmoMode(bool bEnabled);
+
+	/** @return true if Force Combined Gizmo mode is Enabled */
+	bool GetForceCombinedGizmoModeEnabled() const { return bForceCombinedGizmoMode; }
+
+
 protected:
 	/**
 	 * Called upon Level Editor Created in order to bind to EditorElementSelectionPtr changed
@@ -213,6 +224,8 @@ private:
 	// currently defaulting to enabled as FEdModes generally assume this, and in most cases hitproxy pass is not expensive.
 	bool bEnableRenderingDuringHitProxyPass = true;
 
+	bool bForceCombinedGizmoMode = false;
+
 	bool bIsActive = false;
 };
 
diff --git a/Engine/Source/Runtime/InteractiveToolsFramework/Private/BaseGizmos/CombinedTransformGizmo.cpp b/Engine/Source/Runtime/InteractiveToolsFramework/Private/BaseGizmos/CombinedTransformGizmo.cpp
--- a/Engine/Source/Runtime/InteractiveToolsFramework/Private/BaseGizmos/CombinedTransformGizmo.cpp
+++ b/Engine/Source/Runtime/InteractiveToolsFramework/Private/BaseGizmos/CombinedTransformGizmo.cpp
@@ -147,6 +147,7 @@ ACombinedTransformGizmoActor* ACombinedTransformGizmoActor::ConstructCustom3Axis
 		SphereEdge->Radius = 120.0f;
 		SphereEdge->bViewAligned = true;
 		SphereEdge->RegisterComponent();
+		NewActor->RotationSphere = SphereEdge;
 	}
 
 
@@ -443,11 +444,39 @@ void UCombinedTransformGizmo::Tick(float DeltaTime)
 		}
 	}
 
-	bool bShouldShowNonUniformScale = IsNonUniformScaleAllowed();
-	for (UPrimitiveComponent* Component : NonuniformScaleComponents)
+	// apply dynamic visibility filtering to sub-gizmos
+
+	auto SetSubGizmoTypeVisibility = [this](TArray<FSubGizmoInfo>& GizmoInfos, bool bVisible)
+	{
+		for (FSubGizmoInfo& GizmoInfo : GizmoInfos)
+		{
+			if (GizmoInfo.Component.IsValid())
+			{
+				GizmoInfo.Component->SetVisibility(bVisible);
+			}
+		}
+	};
+
+	if (bUseContextGizmoMode)
 	{
-		Component->SetVisibility(bShouldShowNonUniformScale);
+		ActiveGizmoMode = GetGizmoManager()->GetContextQueriesAPI()->GetCurrentTransformGizmoMode();
 	}
+	EToolContextTransformGizmoMode UseGizmoMode = ActiveGizmoMode;
+
+	bool bShouldShowTranslation =
+		(UseGizmoMode == EToolContextTransformGizmoMode::Combined || UseGizmoMode == EToolContextTransformGizmoMode::Translation);
+	bool bShouldShowRotation =
+		(UseGizmoMode == EToolContextTransformGizmoMode::Combined || UseGizmoMode == EToolContextTransformGizmoMode::Rotation);
+	bool bShouldShowUniformScale =
+		(UseGizmoMode == EToolContextTransformGizmoMode::Combined || UseGizmoMode == EToolContextTransformGizmoMode::Scale);
+	bool bShouldShowNonUniformScale = 
+		(UseGizmoMode == EToolContextTransformGizmoMode::Combined || UseGizmoMode == EToolContextTransformGizmoMode::Scale)
+		&& IsNonUniformScaleAllowed();
+
+	SetSubGizmoTypeVisibility(TranslationSubGizmos, bShouldShowTranslation);
+	SetSubGizmoTypeVisibility(RotationSubGizmos, bShouldShowRotation);
+	SetSubGizmoTypeVisibility(UniformScaleSubGizmos, bShouldShowUniformScale);
+	SetSubGizmoTypeVisibility(NonUniformScaleSubGizmos, bShouldShowNonUniformScale);
 
 	UpdateCameraAxisSource();
 }
@@ -495,51 +524,66 @@ void UCombinedTransformGizmo::SetActiveTarget(UTransformProxy* Target, IToolCont
 	// todo should we hold onto these?
 	if (GizmoActor->TranslateX != nullptr)
 	{
-		AddAxisTranslationGizmo(GizmoActor->TranslateX, GizmoComponent, AxisXSource, TransformSource, StateTarget);
+		UInteractiveGizmo* NewGizmo = AddAxisTranslationGizmo(GizmoActor->TranslateX, GizmoComponent, AxisXSource, TransformSource, StateTarget);
 		ActiveComponents.Add(GizmoActor->TranslateX);
+		TranslationSubGizmos.Add( FSubGizmoInfo{ GizmoActor->TranslateX, NewGizmo } );
 	}
 	if (GizmoActor->TranslateY != nullptr)
 	{
-		AddAxisTranslationGizmo(GizmoActor->TranslateY, GizmoComponent, AxisYSource, TransformSource, StateTarget);
+		UInteractiveGizmo* NewGizmo = AddAxisTranslationGizmo(GizmoActor->TranslateY, GizmoComponent, AxisYSource, TransformSource, StateTarget);
 		ActiveComponents.Add(GizmoActor->TranslateY);
+		TranslationSubGizmos.Add(FSubGizmoInfo{ GizmoActor->TranslateY, NewGizmo });
 	}
 	if (GizmoActor->TranslateZ != nullptr)
 	{
-		AddAxisTranslationGizmo(GizmoActor->TranslateZ, GizmoComponent, AxisZSource, TransformSource, StateTarget);
+		UInteractiveGizmo* NewGizmo = AddAxisTranslationGizmo(GizmoActor->TranslateZ, GizmoComponent, AxisZSource, TransformSource, StateTarget);
 		ActiveComponents.Add(GizmoActor->TranslateZ);
+		TranslationSubGizmos.Add(FSubGizmoInfo{ GizmoActor->TranslateZ, NewGizmo });
 	}
 
 
 	if (GizmoActor->TranslateYZ != nullptr)
 	{
-		AddPlaneTranslationGizmo(GizmoActor->TranslateYZ, GizmoComponent, AxisXSource, TransformSource, StateTarget);
+		UInteractiveGizmo* NewGizmo = AddPlaneTranslationGizmo(GizmoActor->TranslateYZ, GizmoComponent, AxisXSource, TransformSource, StateTarget);
 		ActiveComponents.Add(GizmoActor->TranslateYZ);
+		TranslationSubGizmos.Add(FSubGizmoInfo{ GizmoActor->TranslateYZ, NewGizmo });
 	}
 	if (GizmoActor->TranslateXZ != nullptr)
 	{
-		AddPlaneTranslationGizmo(GizmoActor->TranslateXZ, GizmoComponent, AxisYSource, TransformSource, StateTarget);
+		UInteractiveGizmo* NewGizmo = AddPlaneTranslationGizmo(GizmoActor->TranslateXZ, GizmoComponent, AxisYSource, TransformSource, StateTarget);
 		ActiveComponents.Add(GizmoActor->TranslateXZ);
+		TranslationSubGizmos.Add(FSubGizmoInfo{ GizmoActor->TranslateXZ, NewGizmo });
 	}
 	if (GizmoActor->TranslateXY != nullptr)
 	{
-		AddPlaneTranslationGizmo(GizmoActor->TranslateXY, GizmoComponent, AxisZSource, TransformSource, StateTarget);
+		UInteractiveGizmo* NewGizmo = AddPlaneTranslationGizmo(GizmoActor->TranslateXY, GizmoComponent, AxisZSource, TransformSource, StateTarget);
 		ActiveComponents.Add(GizmoActor->TranslateXY);
+		TranslationSubGizmos.Add(FSubGizmoInfo{ GizmoActor->TranslateXY, NewGizmo });
 	}
 
 	if (GizmoActor->RotateX != nullptr)
 	{
-		AddAxisRotationGizmo(GizmoActor->RotateX, GizmoComponent, AxisXSource, TransformSource, StateTarget);
+		UInteractiveGizmo* NewGizmo = AddAxisRotationGizmo(GizmoActor->RotateX, GizmoComponent, AxisXSource, TransformSource, StateTarget);
 		ActiveComponents.Add(GizmoActor->RotateX);
+		RotationSubGizmos.Add(FSubGizmoInfo{ GizmoActor->RotateX, NewGizmo });
 	}
 	if (GizmoActor->RotateY != nullptr)
 	{
-		AddAxisRotationGizmo(GizmoActor->RotateY, GizmoComponent, AxisYSource, TransformSource, StateTarget);
+		UInteractiveGizmo* NewGizmo = AddAxisRotationGizmo(GizmoActor->RotateY, GizmoComponent, AxisYSource, TransformSource, StateTarget);
 		ActiveComponents.Add(GizmoActor->RotateY);
+		RotationSubGizmos.Add(FSubGizmoInfo{ GizmoActor->RotateY, NewGizmo });
 	}
 	if (GizmoActor->RotateZ != nullptr)
 	{
-		AddAxisRotationGizmo(GizmoActor->RotateZ, GizmoComponent, AxisZSource, TransformSource, StateTarget);
+		UInteractiveGizmo* NewGizmo = AddAxisRotationGizmo(GizmoActor->RotateZ, GizmoComponent, AxisZSource, TransformSource, StateTarget);
 		ActiveComponents.Add(GizmoActor->RotateZ);
+		RotationSubGizmos.Add(FSubGizmoInfo{ GizmoActor->RotateZ, NewGizmo });
+	}
+	if (GizmoActor->RotationSphere != nullptr)
+	{
+		// no gizmo for the sphere currently
+		ActiveComponents.Add(GizmoActor->RotationSphere);
+		RotationSubGizmos.Add(FSubGizmoInfo{ GizmoActor->RotationSphere, nullptr });
 	}
 
 
@@ -552,46 +596,47 @@ void UCombinedTransformGizmo::SetActiveTarget(UTransformProxy* Target, IToolCont
 
 	if (GizmoActor->UniformScale != nullptr)
 	{
-		AddUniformScaleGizmo(GizmoActor->UniformScale, GizmoComponent, CameraAxisSource, CameraAxisSource, TransformSource, StateTarget);
+		UInteractiveGizmo* NewGizmo = AddUniformScaleGizmo(GizmoActor->UniformScale, GizmoComponent, CameraAxisSource, CameraAxisSource, TransformSource, StateTarget);
 		ActiveComponents.Add(GizmoActor->UniformScale);
+		UniformScaleSubGizmos.Add(FSubGizmoInfo{ GizmoActor->UniformScale, NewGizmo });
 	}
 
 	if (GizmoActor->AxisScaleX != nullptr)
 	{
-		AddAxisScaleGizmo(GizmoActor->AxisScaleX, GizmoComponent, AxisXSource, UnitAxisXSource, TransformSource, StateTarget);
+		UInteractiveGizmo* NewGizmo = AddAxisScaleGizmo(GizmoActor->AxisScaleX, GizmoComponent, AxisXSource, UnitAxisXSource, TransformSource, StateTarget);
 		ActiveComponents.Add(GizmoActor->AxisScaleX);
-		NonuniformScaleComponents.Add(GizmoActor->AxisScaleX);
+		NonUniformScaleSubGizmos.Add(FSubGizmoInfo{ GizmoActor->AxisScaleX, NewGizmo });
 	}
 	if (GizmoActor->AxisScaleY != nullptr)
 	{
-		AddAxisScaleGizmo(GizmoActor->AxisScaleY, GizmoComponent, AxisYSource, UnitAxisYSource, TransformSource, StateTarget);
+		UInteractiveGizmo* NewGizmo = AddAxisScaleGizmo(GizmoActor->AxisScaleY, GizmoComponent, AxisYSource, UnitAxisYSource, TransformSource, StateTarget);
 		ActiveComponents.Add(GizmoActor->AxisScaleY);
-		NonuniformScaleComponents.Add(GizmoActor->AxisScaleY);
+		NonUniformScaleSubGizmos.Add(FSubGizmoInfo{ GizmoActor->AxisScaleY, NewGizmo });
 	}
 	if (GizmoActor->AxisScaleZ != nullptr)
 	{
-		AddAxisScaleGizmo(GizmoActor->AxisScaleZ, GizmoComponent, AxisZSource, UnitAxisZSource, TransformSource, StateTarget);
+		UInteractiveGizmo* NewGizmo = AddAxisScaleGizmo(GizmoActor->AxisScaleZ, GizmoComponent, AxisZSource, UnitAxisZSource, TransformSource, StateTarget);
 		ActiveComponents.Add(GizmoActor->AxisScaleZ);
-		NonuniformScaleComponents.Add(GizmoActor->AxisScaleZ);
+		NonUniformScaleSubGizmos.Add(FSubGizmoInfo{ GizmoActor->AxisScaleZ, NewGizmo });
 	}
 
 	if (GizmoActor->PlaneScaleYZ != nullptr)
 	{
-		AddPlaneScaleGizmo(GizmoActor->PlaneScaleYZ, GizmoComponent, AxisXSource, UnitAxisXSource, TransformSource, StateTarget);
+		UInteractiveGizmo* NewGizmo = AddPlaneScaleGizmo(GizmoActor->PlaneScaleYZ, GizmoComponent, AxisXSource, UnitAxisXSource, TransformSource, StateTarget);
 		ActiveComponents.Add(GizmoActor->PlaneScaleYZ);
-		NonuniformScaleComponents.Add(GizmoActor->PlaneScaleYZ);
+		NonUniformScaleSubGizmos.Add(FSubGizmoInfo{ GizmoActor->PlaneScaleYZ, NewGizmo });
 	}
 	if (GizmoActor->PlaneScaleXZ != nullptr)
 	{
-		UPlanePositionGizmo* Gizmo = (UPlanePositionGizmo *)AddPlaneScaleGizmo(GizmoActor->PlaneScaleXZ, GizmoComponent, AxisYSource, UnitAxisYSource, TransformSource, StateTarget);
+		UInteractiveGizmo* NewGizmo = AddPlaneScaleGizmo(GizmoActor->PlaneScaleXZ, GizmoComponent, AxisYSource, UnitAxisYSource, TransformSource, StateTarget);
 		ActiveComponents.Add(GizmoActor->PlaneScaleXZ);
-		NonuniformScaleComponents.Add(GizmoActor->PlaneScaleXZ);
+		NonUniformScaleSubGizmos.Add(FSubGizmoInfo{ GizmoActor->PlaneScaleXZ, NewGizmo });
 	}
 	if (GizmoActor->PlaneScaleXY != nullptr)
 	{
-		AddPlaneScaleGizmo(GizmoActor->PlaneScaleXY, GizmoComponent, AxisZSource, UnitAxisZSource, TransformSource, StateTarget);
+		UInteractiveGizmo* NewGizmo = AddPlaneScaleGizmo(GizmoActor->PlaneScaleXY, GizmoComponent, AxisZSource, UnitAxisZSource, TransformSource, StateTarget);
 		ActiveComponents.Add(GizmoActor->PlaneScaleXY);
-		NonuniformScaleComponents.Add(GizmoActor->PlaneScaleXY);
+		NonUniformScaleSubGizmos.Add(FSubGizmoInfo{ GizmoActor->PlaneScaleXY, NewGizmo });
 	}
 }
 
@@ -902,7 +947,10 @@ void UCombinedTransformGizmo::ClearActiveTarget()
 	}
 	ActiveGizmos.SetNum(0);
 	ActiveComponents.SetNum(0);
-	NonuniformScaleComponents.SetNum(0);
+	TranslationSubGizmos.SetNum(0);
+	RotationSubGizmos.SetNum(0);
+	UniformScaleSubGizmos.SetNum(0);
+	NonUniformScaleSubGizmos.SetNum(0);
 
 	CameraAxisSource = nullptr;
 	AxisXSource = nullptr;
diff --git a/Engine/Source/Runtime/InteractiveToolsFramework/Public/BaseGizmos/CombinedTransformGizmo.h b/Engine/Source/Runtime/InteractiveToolsFramework/Public/BaseGizmos/CombinedTransformGizmo.h
--- a/Engine/Source/Runtime/InteractiveToolsFramework/Public/BaseGizmos/CombinedTransformGizmo.h
+++ b/Engine/Source/Runtime/InteractiveToolsFramework/Public/BaseGizmos/CombinedTransformGizmo.h
@@ -89,6 +89,10 @@ public:
 	UPROPERTY()
 	TObjectPtr<UPrimitiveComponent> RotateZ;
 
+	/** Z Axis Rotation Component */
+	UPROPERTY()
+	TObjectPtr<UPrimitiveComponent> RotationSphere;
+
 	//
 	// Scaling Components
 	//
@@ -409,6 +413,21 @@ public:
 	EToolContextCoordinateSystem CurrentCoordinateSystem = EToolContextCoordinateSystem::Local;
 
 
+	/**
+	 * Whether to use the Gizmo Mode provided by the context via the ContextyQueriesAPI.
+	 */
+	UPROPERTY()
+	bool bUseContextGizmoMode = true;
+
+	/**
+	 * Current dynamic sub-widget visibility mode to use (eg Translate-Only, Scale-Only, Combined, etc)
+	 * If bUseContextGizmoMode is true, this value will be updated internally every Tick()
+	 * by quering the ContextyQueriesAPI, otherwise the default is Combined and the client can change it as necessary
+	 */
+	UPROPERTY()
+	EToolContextTransformGizmoMode ActiveGizmoMode = EToolContextTransformGizmoMode::Combined;
+
+
 protected:
 	TSharedPtr<FCombinedTransformGizmoActorFactory> GizmoActorBuilder;
 
@@ -430,18 +449,26 @@ protected:
 	UPROPERTY()
 	TArray<TObjectPtr<UPrimitiveComponent>> ActiveComponents;
 
-	/** 
-	 * List of nonuniform scale components. Subset of of ActiveComponents. These are tracked separately so they can
-	 * be hidden when Gizmo is not configured to use local axes, because UE only supports local nonuniform scaling
-	 * on Components
-	 */
-	UPROPERTY()
-	TArray<TObjectPtr<UPrimitiveComponent>> NonuniformScaleComponents;
-
 	/** list of currently-active child gizmos */
 	UPROPERTY()
 	TArray<TObjectPtr<UInteractiveGizmo>> ActiveGizmos;
 
+	/** 
+	 * FSubGizmoInfo stores the (UPrimitiveComponent,UInteractiveGizmo) pair for a sub-element of the widget.
+	 * The ActiveComponents and ActiveGizmos arrays keep those items alive, so this is redundant information, but useful for filtering/etc
+	 */
+	struct FSubGizmoInfo
+	{
+		// note: either of these may be invalid
+		TWeakObjectPtr<UPrimitiveComponent> Component;
+		TWeakObjectPtr<UInteractiveGizmo> Gizmo;
+	};
+	TArray<FSubGizmoInfo> TranslationSubGizmos;
+	TArray<FSubGizmoInfo> RotationSubGizmos;
+	TArray<FSubGizmoInfo> UniformScaleSubGizmos;
+	TArray<FSubGizmoInfo> NonUniformScaleSubGizmos;
+
+
 	/** GizmoActors will be spawned in this World */
 	UWorld* World;
 
diff --git a/Engine/Source/Runtime/InteractiveToolsFramework/Public/ToolContextInterfaces.h b/Engine/Source/Runtime/InteractiveToolsFramework/Public/ToolContextInterfaces.h
--- a/Engine/Source/Runtime/InteractiveToolsFramework/Public/ToolContextInterfaces.h
+++ b/Engine/Source/Runtime/InteractiveToolsFramework/Public/ToolContextInterfaces.h
@@ -120,6 +120,29 @@ enum class EToolContextCoordinateSystem
 	Local = 1
 };
 
+
+/** 
+ * High-level configuration options for a standard 3D translate/rotate/scale Gizmo, like is commonly used in 3D DCC tools, game editors, etc
+ * This is meant to be used to convey UI-level settings to Tools/Gizmos, eg like the W/E/R toggles for Rranslate/Rotate/Scale in Maya or UE Editor.
+ * More granular control over precise gizmo elements is available through other mechanisms, eg the ETransformGizmoSubElements flags in UCombinedTransformGizmo
+ */
+UENUM()
+enum class EToolContextTransformGizmoMode : uint8
+{
+	/** Hide all Gizmo sub-elements */
+	NoGizmo = 0,
+	/** Only Show Translation sub-elements */
+	Translation = 1,
+	/** Only Show Rotation sub-elements */
+	Rotation = 2,
+	/** Only Show Scale sub-elements */
+	Scale = 3,
+	/** Show all available Gizmo sub-elements */
+	Combined = 8
+};
+
+
+
 /**
  * Snapping configuration settings/state
  */
@@ -164,14 +187,19 @@ public:
 	virtual void GetCurrentViewState(FViewCameraState& StateOut) const = 0;
 
 	/**
-	 * Request current external coordinate-system setting
+	 * Request current external coordinate-system setting. Defaults to World coordinates.
+	 */
+	virtual EToolContextCoordinateSystem GetCurrentCoordinateSystem() const { return EToolContextCoordinateSystem::World; }
+
+	/**
+	 * Request current external Gizmo Mode setting. Defaulting this to Combined gizmo as this was the default behavior before UE-5.2.
 	 */
-	virtual EToolContextCoordinateSystem GetCurrentCoordinateSystem() const = 0;	
+	virtual EToolContextTransformGizmoMode GetCurrentTransformGizmoMode() const { return EToolContextTransformGizmoMode::Combined; }
 
 	/**
-	 * Request current external snapping settings
+	 * Request current external snapping settings. Defaults to no snapping.
 	 */
-	virtual FToolContextSnappingConfiguration GetCurrentSnappingSettings() const = 0;
+	virtual FToolContextSnappingConfiguration GetCurrentSnappingSettings() const { return FToolContextSnappingConfiguration(); }
 
 	/**
 	 * Many tools need standard types of materials that the user should provide (eg a vertex-color material, etc)
