ToolsFramework: InteractiveToolManager::PostActiveToolShutdownRequest() can now be called during UInteractiveTool::Setup(), to cancel further initialization of a Tool. This can be used to (for example) abort Tool startup if it is discovered that necessary data/etc is not available. Added the ability to pass a message to PostActiveToolShutdownRequest(), which will be forwarded to a delegate OnToolUnexpectedShutdownMessage if it is bound.

diff --git a/Engine/Source/Runtime/InteractiveToolsFramework/Private/InteractiveToolManager.cpp b/Engine/Source/Runtime/InteractiveToolsFramework/Private/InteractiveToolManager.cpp
--- a/Engine/Source/Runtime/InteractiveToolsFramework/Private/InteractiveToolManager.cpp
+++ b/Engine/Source/Runtime/InteractiveToolsFramework/Private/InteractiveToolManager.cpp
@@ -207,7 +207,18 @@ bool UInteractiveToolManager::ActivateToolInternal(EToolSide Side)
 	ActiveLeftToolName = ActiveLeftBuilderName;
 	DoPostBuild(Side, ActiveLeftTool, ActiveLeftBuilder, InputState);
 
+	bInToolSetup = true;
 	ActiveLeftTool->Setup();
+	bInToolSetup = false;
+	if (bToolRequestedTerminationDuringSetup)
+	{
+		bToolRequestedTerminationDuringSetup = false;
+		ActiveLeftTool->Shutdown(EToolShutdownType::Cancel);		// need to give Tool a chance to clean up...
+		ActiveLeftTool = nullptr;
+		ActiveLeftToolName.Empty();
+		return false;
+	}
+
 	DoPostSetup(Side, ActiveLeftTool, ActiveLeftBuilder, InputState);
 
 	// register new active input behaviors
@@ -436,8 +447,43 @@ bool UInteractiveToolManager::RequestSelectionChange(const FSelectedOjectsChange
 }
 
 
-bool UInteractiveToolManager::PostActiveToolShutdownRequest(UInteractiveTool* Tool, EToolShutdownType ShutdownType)
+bool UInteractiveToolManager::PostActiveToolShutdownRequest(UInteractiveTool* Tool, EToolShutdownType ShutdownType,
+	bool bShowUnexpectedShutdownMessage, const FText& UnexpectedShutdownMessage)
 {
+	if (bShowUnexpectedShutdownMessage)
+	{
+		if (OnToolUnexpectedShutdownMessage.IsBound() == false )
+		{
+			if ( UnexpectedShutdownMessage.IsEmpty() )
+			{
+				if ( bInToolSetup )
+				{
+					UE_LOG(LogTemp, Error, TEXT("[InteractiveToolManager] Tool %s Could not be Initialized"), *Tool->GetToolInfo().ToolDisplayName.ToString());
+				}
+				else
+				{
+					UE_LOG(LogTemp, Error, TEXT("[InteractiveToolManager] Tool %s was Shut Down Automatically, no message provided"), *Tool->GetToolInfo().ToolDisplayName.ToString());
+				}
+			}
+			else
+			{
+				UE_LOG(LogTemp, Error, TEXT("[InteractiveToolManager] %s"), *UnexpectedShutdownMessage.ToString());
+			}
+		}
+		else
+		{
+			OnToolUnexpectedShutdownMessage.Broadcast(this, Tool, UnexpectedShutdownMessage, bInToolSetup);
+		}
+	}
+
+	// if Tool calls this function from it's Setup() function, then we need to do some special-case handling because
+	// the code below can't be run yet
+	if (bInToolSetup)
+	{
+		bToolRequestedTerminationDuringSetup = true;
+		return true;
+	}
+
 	bool bIsActiveTool = (Tool != nullptr) && (ActiveLeftTool == Tool || ActiveRightTool == Tool);
 	if (!bIsActiveTool)
 	{
diff --git a/Engine/Source/Runtime/InteractiveToolsFramework/Public/InteractiveToolManager.h b/Engine/Source/Runtime/InteractiveToolsFramework/Public/InteractiveToolManager.h
--- a/Engine/Source/Runtime/InteractiveToolsFramework/Public/InteractiveToolManager.h
+++ b/Engine/Source/Runtime/InteractiveToolsFramework/Public/InteractiveToolManager.h
@@ -243,8 +243,19 @@ public:
 	 * The Tool must be active. By default, ::DeactivateTool() wil be called. However if some
 	 * external code has bound to the OnToolShutdownRequest delegate, the request will be
 	 * forwarded to that function first (::DeactivateTool() will still be called if it returns false)
+	 * 
+	 * If this function is called during InteractiveTool::Setup(), the Tool invocation will be aborted,
+	 * ie the Tool will be immediately shutdown with EToolShutdownType::Cancel, and never become "Active"/etc. 
+	 * Note also that in this case the shutdown request *will not* be forwarded to the OnToolShutdownRequest delegate.
+	 * 
+	 * @param bShowUnexpectedShutdownMessage if this was an unexpected shutdown, passing true here, along with a non-empty UnexpectedShutdownMessage, will result in OnToolUnexpectedShutdownMessage being called
+	 * @param UnexpectedShutdownMessage message provided by the caller meant to explain why the Tool has unexpectedly been shut down
 	 */
-	virtual bool PostActiveToolShutdownRequest(UInteractiveTool* Tool, EToolShutdownType ShutdownType);
+	virtual bool PostActiveToolShutdownRequest(
+		UInteractiveTool* Tool, 
+		EToolShutdownType ShutdownType,
+		bool bShowUnexpectedShutdownMessage = false,
+		const FText& UnexpectedShutdownMessage = FText() );
 
 	DECLARE_DELEGATE_RetVal_ThreeParams(bool, FOnToolShutdownRequest, UInteractiveToolManager*, UInteractiveTool*, EToolShutdownType);
 	/**
@@ -253,8 +264,15 @@ public:
 	 * that the request will be handled, otherwise ::DeactivateTool() will be called
 	 */
 	FOnToolShutdownRequest OnToolShutdownRequest;
-	 
+	
+	DECLARE_MULTICAST_DELEGATE_FourParams(FOnToolUnexpectedShutdownMessage, UInteractiveToolManager*, UInteractiveTool*, const FText& Message, bool bWasDuringToolSetup);
 
+	/**
+	 * PostActiveToolShutdownRequest() will broadcast this delegate with a Message if bShowUnexpectedShutdownMessage=true.
+	 * If the shutdown was during Tool Setup/startup, bWasDuringToolSetup will be passed as true
+	 */
+	FOnToolUnexpectedShutdownMessage OnToolUnexpectedShutdownMessage;
+	
 
 	//
 	// access to APIs, etc
@@ -311,6 +329,12 @@ protected:
 	/** This flag is set to true on Initialize() and false on Shutdown(). */
 	bool bIsActive = false;
 
+	// This bool is for handling the case where a Tool posts a shutdown request during it's Setup() call.
+	// In that case we will terminate the Tool immediately
+	bool bInToolSetup = false;
+	bool bToolRequestedTerminationDuringSetup = false;
+
+
 	/** Current set of named ToolBuilders */
 	UPROPERTY()
 	TMap<FString, TObjectPtr<UInteractiveToolBuilder>> ToolBuilders;
