resolve PVS warning

diff --git a/Engine/Source/Runtime/InteractiveToolsFramework/Private/InteractiveToolManager.cpp b/Engine/Source/Runtime/InteractiveToolsFramework/Private/InteractiveToolManager.cpp
--- a/Engine/Source/Runtime/InteractiveToolsFramework/Private/InteractiveToolManager.cpp
+++ b/Engine/Source/Runtime/InteractiveToolsFramework/Private/InteractiveToolManager.cpp
@@ -456,13 +456,14 @@ bool UInteractiveToolManager::PostActiveToolShutdownRequest(UInteractiveTool* To
 		{
 			if ( UnexpectedShutdownMessage.IsEmpty() )
 			{
+				FString ToolName = (Tool != nullptr) ? Tool->GetToolInfo().ToolDisplayName.ToString() : FString(TEXT("(Null Tool)"));
 				if ( bInToolSetup )
 				{
-					UE_LOG(LogTemp, Error, TEXT("[InteractiveToolManager] Tool %s Could not be Initialized"), *Tool->GetToolInfo().ToolDisplayName.ToString());
+					UE_LOG(LogTemp, Error, TEXT("[InteractiveToolManager] Tool %s Could not be Initialized"), *ToolName);
 				}
 				else
 				{
-					UE_LOG(LogTemp, Error, TEXT("[InteractiveToolManager] Tool %s was Shut Down Automatically, no message provided"), *Tool->GetToolInfo().ToolDisplayName.ToString());
+					UE_LOG(LogTemp, Error, TEXT("[InteractiveToolManager] Tool %s was Shut Down Automatically, no message provided"), *ToolName);
 				}
 			}
 			else
