ToolsFramework: Hid various internal tools framework Actors from placement and blueprinting. Added some markup to various InteractiveTool UProperties to prevent them from ever being serialized, for incoming Blueprintable Tools.
remove unnecessary EToolShutdownType predeclaration from BaseToolkit.h, this was added a long time to ago when the type was referenced in this header, but that code was moved w/o removing the predecl

diff --git a/Engine/Plugins/Runtime/MeshModelingToolset/Source/ModelingComponents/Public/Drawing/PreviewGeometryActor.h b/Engine/Plugins/Runtime/MeshModelingToolset/Source/ModelingComponents/Public/Drawing/PreviewGeometryActor.h
--- a/Engine/Plugins/Runtime/MeshModelingToolset/Source/ModelingComponents/Public/Drawing/PreviewGeometryActor.h
+++ b/Engine/Plugins/Runtime/MeshModelingToolset/Source/ModelingComponents/Public/Drawing/PreviewGeometryActor.h
@@ -12,7 +12,7 @@
 /**
  * An actor suitable for attaching components used to draw preview elements, such as LineSetComponent and TriangleSetComponent.
  */
-UCLASS()
+UCLASS(Transient, NotPlaceable, Hidden, NotBlueprintable, NotBlueprintType)
 class MODELINGCOMPONENTS_API APreviewGeometryActor : public AInternalToolFrameworkActor
 {
 	GENERATED_BODY()
diff --git a/Engine/Plugins/Runtime/MeshModelingToolset/Source/ModelingComponents/Public/PreviewMesh.h b/Engine/Plugins/Runtime/MeshModelingToolset/Source/ModelingComponents/Public/PreviewMesh.h
--- a/Engine/Plugins/Runtime/MeshModelingToolset/Source/ModelingComponents/Public/PreviewMesh.h
+++ b/Engine/Plugins/Runtime/MeshModelingToolset/Source/ModelingComponents/Public/PreviewMesh.h
@@ -22,7 +22,7 @@ PREDECLARE_GEOMETRY(template<typename RealType> class TMeshTangents);
  * We use this AInternalToolFrameworkActor subclass so that we can identify such objects
  * at higher levels (for example to prevent them from being deleted in the Editor)
  */
-UCLASS()
+UCLASS(Transient, NotPlaceable, Hidden, NotBlueprintable, NotBlueprintType)
 class MODELINGCOMPONENTS_API APreviewMeshActor : public AInternalToolFrameworkActor
 {
 	GENERATED_BODY()
diff --git a/Engine/Source/Editor/UnrealEd/Public/Toolkits/BaseToolkit.h b/Engine/Source/Editor/UnrealEd/Public/Toolkits/BaseToolkit.h
--- a/Engine/Source/Editor/UnrealEd/Public/Toolkits/BaseToolkit.h
+++ b/Engine/Source/Editor/UnrealEd/Public/Toolkits/BaseToolkit.h
@@ -17,7 +17,6 @@ class UEdMode;
 class IDetailsView;
 struct FDetailsViewArgs;
 class FAssetEditorModeUILayer;
-enum class EToolShutdownType;
 class SBorder;
 
 /**
diff --git a/Engine/Source/Runtime/InteractiveToolsFramework/Private/InteractiveTool.cpp b/Engine/Source/Runtime/InteractiveToolsFramework/Private/InteractiveTool.cpp
--- a/Engine/Source/Runtime/InteractiveToolsFramework/Private/InteractiveTool.cpp
+++ b/Engine/Source/Runtime/InteractiveToolsFramework/Private/InteractiveTool.cpp
@@ -189,6 +189,11 @@ void UInteractiveToolPropertySet::SaveRestoreProperties(UInteractiveTool* Restor
 		// if this is the first time we have seen this property set, then we don't have any values to Restore
 		return;
 	}
+	if (PropertyCache == nullptr)
+	{
+		// if for whatever reason a valid PropertyCache could not be returned, just abort
+		return;
+	}
 	for ( FProperty* Prop : TFieldRange<FProperty>(GetClass()) )
 	{
 #if WITH_EDITOR
@@ -206,6 +211,35 @@ void UInteractiveToolPropertySet::SaveRestoreProperties(UInteractiveTool* Restor
 	}
 }
 
+
+TObjectPtr<UInteractiveToolPropertySet> UInteractiveToolPropertySet::GetDynamicPropertyCache(const FString& CacheIdentifier, bool& bWasCreatedOut)
+{
+	bWasCreatedOut = false;
+	UInteractiveToolPropertySet* CDO = GetMutableDefault<UInteractiveToolPropertySet>(GetClass());
+	TObjectPtr<UInteractiveToolPropertySet>* Found = CDO->CachedPropertiesMap.Find(CacheIdentifier);
+	if (Found == nullptr)
+	{
+		TObjectPtr<UInteractiveToolPropertySet> NewPropCache = NewObject<UInteractiveToolPropertySet>((UObject*)GetTransientPackage(), GetClass());
+		ensure(NewPropCache != nullptr);
+		CDO->CachedPropertiesMap.Add(CacheIdentifier, NewPropCache);
+		bWasCreatedOut = true;
+		return NewPropCache;
+	}
+
+	if ( ensure(*Found != nullptr) == false )
+	{
+		// this case seems to occur sometimes for Blueprintable Tools, uncertain why, but perhaps this ensure will help to find it
+		TObjectPtr<UInteractiveToolPropertySet> NewPropCache = NewObject<UInteractiveToolPropertySet>((UObject*)GetTransientPackage(), GetClass());
+		ensure(NewPropCache != nullptr);
+		CDO->CachedPropertiesMap[CacheIdentifier] = NewPropCache;
+		bWasCreatedOut = true;
+		return NewPropCache;
+	}
+
+	return *Found;
+}
+
+
 void UInteractiveTool::RegisterActions(FInteractiveToolActionSet& ActionSet)
 {
 }
diff --git a/Engine/Source/Runtime/InteractiveToolsFramework/Public/BaseGizmos/CombinedTransformGizmo.h b/Engine/Source/Runtime/InteractiveToolsFramework/Public/BaseGizmos/CombinedTransformGizmo.h
--- a/Engine/Source/Runtime/InteractiveToolsFramework/Public/BaseGizmos/CombinedTransformGizmo.h
+++ b/Engine/Source/Runtime/InteractiveToolsFramework/Public/BaseGizmos/CombinedTransformGizmo.h
@@ -35,7 +35,7 @@ class FTransformGizmoTransformChange;
  * The static factory method ::ConstructDefault3AxisGizmo() creates and initializes an 
  * Actor suitable for use in a standard 3-axis Transformation Gizmo.
  */
-UCLASS(Transient)
+UCLASS(Transient, NotPlaceable, Hidden, NotBlueprintable, NotBlueprintType)
 class INTERACTIVETOOLSFRAMEWORK_API ACombinedTransformGizmoActor : public AGizmoActor
 {
 	GENERATED_BODY()
diff --git a/Engine/Source/Runtime/InteractiveToolsFramework/Public/BaseGizmos/GizmoActor.h b/Engine/Source/Runtime/InteractiveToolsFramework/Public/BaseGizmos/GizmoActor.h
--- a/Engine/Source/Runtime/InteractiveToolsFramework/Public/BaseGizmos/GizmoActor.h
+++ b/Engine/Source/Runtime/InteractiveToolsFramework/Public/BaseGizmos/GizmoActor.h
@@ -19,7 +19,7 @@ class UGizmoViewContext;
  * Currently this does not involve any special functionality, but a set of static functions
  * are provided to create default Components commonly used in Gizmos.
  */
-UCLASS(Transient)
+UCLASS(Transient, NotPlaceable, Hidden, NotBlueprintable, NotBlueprintType)
 class INTERACTIVETOOLSFRAMEWORK_API AGizmoActor : public AInternalToolFrameworkActor
 {
 	GENERATED_BODY()
diff --git a/Engine/Source/Runtime/InteractiveToolsFramework/Public/BaseGizmos/IntervalGizmo.h b/Engine/Source/Runtime/InteractiveToolsFramework/Public/BaseGizmos/IntervalGizmo.h
--- a/Engine/Source/Runtime/InteractiveToolsFramework/Public/BaseGizmos/IntervalGizmo.h
+++ b/Engine/Source/Runtime/InteractiveToolsFramework/Public/BaseGizmos/IntervalGizmo.h
@@ -23,7 +23,7 @@ class UGizmoViewContext;
 
 
 /**
- *  AIntervalGizmoActor is an Actor type intended to be used with UIntervalGizmo,
+ * AIntervalGizmoActor is an Actor type intended to be used with UIntervalGizmo,
  * as the in-scene visual representation of the Gizmo.
  *
  * FIntervalGizmoActorFactory returns an instance of this Actor type (or a subclass).
@@ -31,7 +31,7 @@ class UGizmoViewContext;
  * If a particular sub-Gizmo is not required, simply set that UProperty to null.
 
  */
-UCLASS(Transient)
+UCLASS(Transient, NotPlaceable, Hidden, NotBlueprintable, NotBlueprintType)
 class INTERACTIVETOOLSFRAMEWORK_API AIntervalGizmoActor : public AGizmoActor
 {
 	GENERATED_BODY()
diff --git a/Engine/Source/Runtime/InteractiveToolsFramework/Public/InteractiveTool.h b/Engine/Source/Runtime/InteractiveToolsFramework/Public/InteractiveTool.h
--- a/Engine/Source/Runtime/InteractiveToolsFramework/Public/InteractiveTool.h
+++ b/Engine/Source/Runtime/InteractiveToolsFramework/Public/InteractiveTool.h
@@ -39,7 +39,8 @@ class UInteractiveToolManager;
 struct FPropertyChangedEvent;
 
 /** Passed to UInteractiveTool::Shutdown to indicate how Tool should shut itself down*/
-enum class EToolShutdownType
+UENUM(BlueprintType)
+enum class EToolShutdownType : uint8
 {
 	/** Tool cleans up and exits. Pass this to tools that do not have Accept/Cancel options. */
 	Completed = 0,
@@ -299,20 +300,7 @@ protected:
 	 * @param bWasCreatedOut true is returned here if this is the first time the object was seen
 	 * @return instance of the current subclass that can be used to save/restore values
 	 */
-	TObjectPtr<UInteractiveToolPropertySet> GetDynamicPropertyCache(const FString& CacheIdentifier, bool& bWasCreatedOut)
-	{
-		bWasCreatedOut = false;
-		UInteractiveToolPropertySet* CDO = GetMutableDefault<UInteractiveToolPropertySet>(GetClass());
-		TObjectPtr<UInteractiveToolPropertySet>* Found = CDO->CachedPropertiesMap.Find(CacheIdentifier);
-		if (Found == nullptr)
-		{
-			TObjectPtr<UInteractiveToolPropertySet> NewPropCache = NewObject<UInteractiveToolPropertySet>((UObject*)GetTransientPackage(), GetClass());
-			CDO->CachedPropertiesMap.Add(CacheIdentifier, NewPropCache);
-			bWasCreatedOut = true;
-			return NewPropCache;
-		}
-		return *Found;
-	}
+	TObjectPtr<UInteractiveToolPropertySet> GetDynamicPropertyCache(const FString& CacheIdentifier, bool& bWasCreatedOut);
 
 public:
 #if WITH_EDITOR
@@ -325,11 +313,11 @@ public:
 #endif
 
 protected:
-	UPROPERTY()
+	UPROPERTY(Transient, DuplicateTransient, NonTransactional, SkipSerialization)
 	TMap<FString, TObjectPtr<UInteractiveToolPropertySet>> CachedPropertiesMap;
 
 	// Controls whether a property set is shown in the UI.  Transient so that disabling a PropertySet in one tool doesn't disable it in others.
-	UPROPERTY(meta=(TransientToolProperty))
+	UPROPERTY(Transient, DuplicateTransient, SkipSerialization, meta=(TransientToolProperty))
 	bool bIsPropertySetEnabled = true;
 
 	friend class UInteractiveTool;	// so that tool can enable/disable
@@ -484,11 +472,11 @@ public:
 protected:
 
 	/** The current set of InputBehaviors provided by this Tool */
-	UPROPERTY()
+	UPROPERTY(Transient, DuplicateTransient, NonTransactional, SkipSerialization)
 	TObjectPtr<UInputBehaviorSet> InputBehaviors;
 
 	/** The current set of Property UObjects provided by this Tool. May contain pointer to itself. */
-	UPROPERTY()
+	UPROPERTY(Transient, DuplicateTransient, NonTransactional, SkipSerialization)
 	TArray<TObjectPtr<UObject>> ToolPropertyObjects;
 
 	/**
@@ -551,21 +539,23 @@ protected:
 	virtual void UpdateAcceptWarnings(EAcceptWarning Warning);
 
 
-	//
-	// Action support/system
-	//
-	// Your Tool subclass can register a set of "Actions" it can execute
-	// by overloading RegisterActions(). Then external systems can use GetActionSet() to
-	// find out what Actions your Tool supports, and ExecuteAction() to run those actions.
-	//
 
-private:
+protected:
 	/**
 	 * Allow the Tool to do any necessary processing on Tick
 	 * @param DeltaTime the time delta since last tick
 	 */
 	virtual void OnTick(float DeltaTime){};
 
+
+
+	//
+	// Action support/system
+	//
+	// Your Tool subclass can register a set of "Actions" it can execute
+	// by overloading RegisterActions(). Then external systems can use GetActionSet() to
+	// find out what Actions your Tool supports, and ExecuteAction() to run those actions.
+	//
 public:
 	/**
 	 * Get the internal Action Set for this Tool. The action set is created and registered on-demand.
