ToolsFramework: Expose FInputRayHit and FInputDeviceRay to Blueprints, in preparation for Tool BPs. Change FInputRayHit.HitDepth to be a double, fix up FInputCaptureRequest.

diff --git a/Engine/Source/Runtime/InteractiveToolsFramework/Public/InputBehavior.h b/Engine/Source/Runtime/InteractiveToolsFramework/Public/InputBehavior.h
--- a/Engine/Source/Runtime/InteractiveToolsFramework/Public/InputBehavior.h
+++ b/Engine/Source/Runtime/InteractiveToolsFramework/Public/InputBehavior.h
@@ -79,13 +79,13 @@ struct INTERACTIVETOOLSFRAMEWORK_API FInputCaptureRequest
 	EInputCaptureSide Side;
 
 	/** Depth along hit-test ray */
-	float HitDepth;
+	double HitDepth;
 
 	/** Owner of the requesting behavior. Behavior doesn't know this, so this is initialized to null */
 	void* Owner;
 
 
-	FInputCaptureRequest(EInputCaptureRequestType type, UInputBehavior* behavior, EInputCaptureSide whichSide, float hitDepth = TNumericLimits<float>::Max() )
+	FInputCaptureRequest(EInputCaptureRequestType type, UInputBehavior* behavior, EInputCaptureSide whichSide, double hitDepth = TNumericLimits<double>::Max() )
 	{
 		this->Type = type;
 		this->Source = behavior;
@@ -95,7 +95,7 @@ struct INTERACTIVETOOLSFRAMEWORK_API FInputCaptureRequest
 	}
 
 	/** Create a Begin-capture request */
-	static FInputCaptureRequest Begin(UInputBehavior* behavior, EInputCaptureSide whichSide, float hitDepth = TNumericLimits<float>::Max() )
+	static FInputCaptureRequest Begin(UInputBehavior* behavior, EInputCaptureSide whichSide, double hitDepth = TNumericLimits<double>::Max() )
 	{
 		return FInputCaptureRequest(EInputCaptureRequestType::Begin, behavior, whichSide, hitDepth);
 	}
@@ -103,7 +103,7 @@ struct INTERACTIVETOOLSFRAMEWORK_API FInputCaptureRequest
 	/** Create an ignore-capture request */
 	static FInputCaptureRequest Ignore() 
 	{
-		return FInputCaptureRequest(EInputCaptureRequestType::Ignore, nullptr, EInputCaptureSide::Any, TNumericLimits<float>::Max() );
+		return FInputCaptureRequest(EInputCaptureRequestType::Ignore, nullptr, EInputCaptureSide::Any, TNumericLimits<double>::Max() );
 	}
 
 	friend bool operator<(const FInputCaptureRequest& l, const FInputCaptureRequest& r);
diff --git a/Engine/Source/Runtime/InteractiveToolsFramework/Public/InputState.h b/Engine/Source/Runtime/InteractiveToolsFramework/Public/InputState.h
--- a/Engine/Source/Runtime/InteractiveToolsFramework/Public/InputState.h
+++ b/Engine/Source/Runtime/InteractiveToolsFramework/Public/InputState.h
@@ -47,24 +47,29 @@ ENUM_CLASS_FLAGS(EInputDevices);
  * FInputRayHit is returned by various hit-test interface functions.
  * Generally this is intended to be returned as the result of a hit-test with a FInputDeviceRay 
  */
-USTRUCT()
+USTRUCT(BlueprintType)
 struct INTERACTIVETOOLSFRAMEWORK_API FInputRayHit
 {
 	GENERATED_BODY()
 
 	/** true if ray hit something, false otherwise */
+	UPROPERTY(BlueprintReadWrite, Category = InputRayHit)
 	bool bHit;
 
 	/** distance along ray at which intersection occurred */
-	float HitDepth;
+	UPROPERTY(BlueprintReadWrite, Category = InputRayHit)
+	double HitDepth;
 
 	/** Normal at hit point, if available */
+	UPROPERTY(BlueprintReadWrite, Category = InputRayHit)
 	FVector HitNormal;
 
 	/** True if HitNormal was set */
+	UPROPERTY(BlueprintReadWrite, Category = InputRayHit)
 	bool bHasHitNormal;
 
 	/** Client-defined integer identifier for hit object/element/target/etc */
+	UPROPERTY(BlueprintReadWrite, Category = InputRayHit)
 	int32 HitIdentifier;
 
 	/**
@@ -77,6 +82,7 @@ struct INTERACTIVETOOLSFRAMEWORK_API FInputRayHit
 	  * Client-defined pointer for UObject-derived hit owners.  
 	  * HitOwner and HitObject should be set to the same pointer if the HitOwner derives from UObject. 
 	  */
+	UPROPERTY(BlueprintReadWrite, Category = InputRayHit)
 	TWeakObjectPtr<UObject> HitObject;
 
 	/** Set hit object, will also set hit owner to the same value */
@@ -89,7 +95,7 @@ struct INTERACTIVETOOLSFRAMEWORK_API FInputRayHit
 	FInputRayHit()
 	{
 		bHit = false;
-		HitDepth = TNumericLimits<float>::Max();
+		HitDepth = (double)TNumericLimits<float>::Max();
 		HitNormal = FVector(0, 0, 1);
 		bHasHitNormal = false;
 		HitIdentifier = 0;
@@ -97,7 +103,7 @@ struct INTERACTIVETOOLSFRAMEWORK_API FInputRayHit
 		HitObject = nullptr;
 	}
 
-	explicit FInputRayHit(float HitDepthIn)
+	explicit FInputRayHit(double HitDepthIn)
 	{
 		bHit = true;
 		HitDepth = HitDepthIn;
@@ -108,7 +114,7 @@ struct INTERACTIVETOOLSFRAMEWORK_API FInputRayHit
 		HitObject = nullptr;
 	}
 
-	explicit FInputRayHit(float HitDepthIn, const FVector& HitNormalIn)
+	explicit FInputRayHit(double HitDepthIn, const FVector& HitNormalIn)
 	{
 		bHit = true;
 		HitDepth = HitDepthIn;
@@ -118,6 +124,7 @@ struct INTERACTIVETOOLSFRAMEWORK_API FInputRayHit
 		HitOwner = nullptr;
 		HitObject = nullptr;
 	}
+
 };
 
 
@@ -315,17 +322,30 @@ struct FInputDeviceState
  * If the device is a 2D input device like a mouse, then the ray may
  * have an associated 2D screen position.
  */
-struct FInputDeviceRay
+USTRUCT(BlueprintType)
+struct INTERACTIVETOOLSFRAMEWORK_API FInputDeviceRay
 {
+	GENERATED_BODY()
+
 	/** 3D ray in 3D scene, in world coordinates */
+	UPROPERTY(BlueprintReadWrite, Category = InputDeviceRay)
 	FRay WorldRay;
 
 	/** If true, WorldRay has 2D device position coordinates */
+	UPROPERTY(BlueprintReadWrite, Category = InputDeviceRay)
 	bool bHas2D = false;
 
 	/** 2D device position coordinates associated with the ray */
+	UPROPERTY(BlueprintReadWrite, Category = InputDeviceRay)
 	FVector2D ScreenPosition;
 
+	// this is required for a USTRUCT
+	FInputDeviceRay()
+	{
+		WorldRay = FRay();
+		bHas2D = false;
+		ScreenPosition = FVector2D(0, 0);
+	}
 
 	explicit FInputDeviceRay(const FRay& WorldRayIn)
 	{
