ToolsFramework: allow InputRouter to update the modifier-key states on active mouse capture when the modifier-keys change but the mouse hasn't moved

diff --git a/Engine/Source/Runtime/InteractiveToolsFramework/Private/InputRouter.cpp b/Engine/Source/Runtime/InteractiveToolsFramework/Private/InputRouter.cpp
--- a/Engine/Source/Runtime/InteractiveToolsFramework/Private/InputRouter.cpp
+++ b/Engine/Source/Runtime/InteractiveToolsFramework/Private/InputRouter.cpp
@@ -60,10 +60,26 @@ bool UInputRouter::PostInputEvent(const FInputDeviceState& Input)
 	if (Input.IsFromDevice(EInputDevices::Mouse))
 	{
 		PostInputEvent_Mouse(Input);
+		LastMouseInputState = Input;
 		return HasActiveMouseCapture();
 	}
 	else if (Input.IsFromDevice(EInputDevices::Keyboard))
 	{
+		// if we are actively capturing Mouse and the key event is from a modifier key, 
+		// we want to update those modifiers
+		if ( (HasActiveMouseCapture() || ActiveLeftHoverCapture != nullptr)
+			   && Input.Keyboard.ActiveKey.Button.IsModifierKey() )
+		{
+			if ((LastMouseInputState.bAltKeyDown != Input.bAltKeyDown) ||
+				(LastMouseInputState.bShiftKeyDown != Input.bShiftKeyDown) ||
+				(LastMouseInputState.bCtrlKeyDown != Input.bCtrlKeyDown) ||
+				(LastMouseInputState.bCmdKeyDown != Input.bCmdKeyDown))
+			{
+				LastMouseInputState.SetModifierKeyStates(Input.bShiftKeyDown, Input.bAltKeyDown, Input.bCtrlKeyDown, Input.bCmdKeyDown);
+				PostInputEvent_Mouse(LastMouseInputState);
+			}
+		}
+
 		PostInputEvent_Keyboard(Input);
 		return (ActiveKeyboardCapture != nullptr);
 	}
diff --git a/Engine/Source/Runtime/InteractiveToolsFramework/Public/InputRouter.h b/Engine/Source/Runtime/InteractiveToolsFramework/Public/InputRouter.h
--- a/Engine/Source/Runtime/InteractiveToolsFramework/Public/InputRouter.h
+++ b/Engine/Source/Runtime/InteractiveToolsFramework/Public/InputRouter.h
@@ -111,7 +111,6 @@ protected:
 	void CheckForMouseCaptures(const FInputDeviceState& Input);
 	void HandleCapturedMouseInput(const FInputDeviceState& Input);
 
-
 	//
 	// Hover support
 	//
@@ -119,6 +118,8 @@ protected:
 	UInputBehavior* ActiveLeftHoverCapture = nullptr;
 	void* ActiveLeftHoverCaptureOwner = nullptr;
 
+	FInputDeviceState LastMouseInputState;
+
 	void TerminateHover(EInputCaptureSide Side);
 	bool ProcessMouseHover(const FInputDeviceState& Input);
 };
\ No newline at end of file
