ToolsFramework: Support optional Relative translation snapping (in World & Local coordinates) in CombinedTransformGizmo while still optionally supporting Absolute translation snapping in World coordinates. Previously only supported Absolute snapping in World coordinates. Relative is now the default, similar to standard UE gizmo.

- moved most functions defined in ParameterToTransformAdapters.h header into a new cpp to simplify debugging/etc
- UGizmoAxisTranslationParameterSource and UGizmoPlaneTranslationParameterSource now support optional constraints on the translation Delta
- CombinedTransformGizmo now has PositionAxisDeltaSnapFunction that, when in 'relative' mode, snaps the translation Deltas to increments via those constraints
- Added CombinedTransformGizmo::RelativeTranslationSnapping which controls relative vs absolute translation snapping, and can fetch value from ToolsContext
- Also added rotation-angle-delta constraint support UGizmoAxisRotationParameterSource and switched the gizmo to explicitly snap angle-deltas instead of trying to snap to "world rotation grid"
- Moved SnapToIncrement function to UE::Geometry namespace (in VectorUtil.h) and cleaned up a bit of duplicate code in UModelingSceneSnappingManager left over from pre-GeometryCore times

- Added FToolContextSnappingConfiguration::bEnableAbsoluteWorldSnapping, EdModeInteractiveToolsContext:: bEnableAbsoluteWorldSnapping and bits to hook this up
- Added UModelingToolsEditorModeSettings::bEnableAbsoluteWorldSnapping which stores this setting globally, but currently it is not serialized to config file, resets to default (Relative) on Editor restart
- Toolkit exposes toggle in the Selection toolbar menus

diff --git a/Engine/Plugins/Editor/ModelingToolsEditorMode/Source/ModelingToolsEditorMode/Private/ModelingModeToolkit_Toolbars.cpp b/Engine/Plugins/Editor/ModelingToolsEditorMode/Source/ModelingToolsEditorMode/Private/ModelingModeToolkit_Toolbars.cpp
--- a/Engine/Plugins/Editor/ModelingToolsEditorMode/Source/ModelingToolsEditorMode/Private/ModelingModeToolkit_Toolbars.cpp
+++ b/Engine/Plugins/Editor/ModelingToolsEditorMode/Source/ModelingToolsEditorMode/Private/ModelingModeToolkit_Toolbars.cpp
@@ -76,6 +76,7 @@ void MakeSubMenu_GizmoVisibilityMode(FModelingToolsEditorModeToolkit* Toolkit, F
 	UModelingToolsEditorModeSettings* Settings = GetMutableDefault<UModelingToolsEditorModeSettings>();
 	UEditorInteractiveToolsContext* Context = Toolkit->GetScriptableEditorMode()->GetInteractiveToolsContext(EToolsContextScope::EdMode);
 
+	// toggle for Combined/Separate Gizmo Mode
 	const FUIAction GizmoMode_Combined(
 		FExecuteAction::CreateLambda([Settings, Context]
 		{
@@ -91,6 +92,22 @@ void MakeSubMenu_GizmoVisibilityMode(FModelingToolsEditorModeToolkit* Toolkit, F
 	MenuBuilder.AddMenuEntry(LOCTEXT("GizmoMode_Combined", "Combined Gizmo"), 
 		LOCTEXT("GizmoMode_Combined_Tooltip", "Ignore Level Editor Gizmo Mode and always use a Combined Transform Gizmo in Modeling Tools"),
 		FSlateIcon(), GizmoMode_Combined, NAME_None, EUserInterfaceActionType::ToggleButton);
+
+	// toggle for Absolute Grid Snapping mode in World Coordinates
+	const FUIAction GizmoMode_AbsoluteWorldSnap(
+		FExecuteAction::CreateLambda([Settings, Context]
+		{
+			Context->SetAbsoluteWorldSnappingEnabled( ! Context->GetAbsoluteWorldSnappingEnabled() );
+			Settings->bEnableAbsoluteWorldSnapping = Context->GetAbsoluteWorldSnappingEnabled();
+		}),
+		FCanExecuteAction(),
+		FIsActionChecked::CreateLambda([Context]()
+		{
+			return Context->GetAbsoluteWorldSnappingEnabled();
+		}));
+	MenuBuilder.AddMenuEntry(LOCTEXT("GizmoMode_AbsoluteWorldSnap", "World Grid Snapping"), 
+		LOCTEXT("GizmoMode_AbsoluteWorldSnap_Tooltip", "Snap Translation/Rotation to Absolute Grid Coordinates in the World Coordinate System, instead of Relative to the initial position"),
+		FSlateIcon(), GizmoMode_AbsoluteWorldSnap, NAME_None, EUserInterfaceActionType::ToggleButton);
 }
 
 
diff --git a/Engine/Plugins/Editor/ModelingToolsEditorMode/Source/ModelingToolsEditorMode/Private/ModelingSelectionInteraction.cpp b/Engine/Plugins/Editor/ModelingToolsEditorMode/Source/ModelingToolsEditorMode/Private/ModelingSelectionInteraction.cpp
--- a/Engine/Plugins/Editor/ModelingToolsEditorMode/Source/ModelingToolsEditorMode/Private/ModelingSelectionInteraction.cpp
+++ b/Engine/Plugins/Editor/ModelingToolsEditorMode/Source/ModelingToolsEditorMode/Private/ModelingSelectionInteraction.cpp
@@ -338,6 +338,10 @@ void UModelingSelectionInteraction::OnBeginGizmoTransform(UTransformProxy* Proxy
 	InitialGizmoFrame = FFrame3d(Transform);
 	InitialGizmoScale = FVector3d(Transform.GetScale3D());
 
+	LastTranslationDelta = FVector3d::Zero();
+	LastRotateDelta = FVector4d::Zero();
+	LastScaleDelta = FVector3d::Zero();
+
 	bInActiveTransform = SelectionManager->BeginTransformation();
 
 	OnTransformBegin.Broadcast();
@@ -380,8 +384,15 @@ void UModelingSelectionInteraction::ApplyPendingTransformInteractions()
 
 	const bool bLastUpdateUsedWorldFrame = true;	// for later local-frame support
 
-	if (TranslationDelta.SquaredLength() > 0.0001 || RotateDelta.SquaredLength() > 0.0001 || CurScaleDelta.SquaredLength() > 0.0001)
+	// if any of the deltas have deviated from the last delta, forward the transformation change on to targets
+	if ( (TranslationDelta - LastTranslationDelta).SizeSquared() > FMathf::ZeroTolerance
+		|| ((FVector4d)RotateDelta - LastRotateDelta).SizeSquared() > FMathf::ZeroTolerance
+		|| (CurScaleDelta- LastScaleDelta).SizeSquared() > FMathf::ZeroTolerance)
 	{
+		LastTranslationDelta = TranslationDelta;
+		LastRotateDelta = (FVector4d)RotateDelta;
+		LastScaleDelta = CurScaleDelta;
+
 		if (bLastUpdateUsedWorldFrame)
 		{
 			// For a world frame gizmo, the scaling needs to happen in world aligned gizmo space, but the 
diff --git a/Engine/Plugins/Editor/ModelingToolsEditorMode/Source/ModelingToolsEditorMode/Private/ModelingToolsEditorMode.cpp b/Engine/Plugins/Editor/ModelingToolsEditorMode/Source/ModelingToolsEditorMode/Private/ModelingToolsEditorMode.cpp
--- a/Engine/Plugins/Editor/ModelingToolsEditorMode/Source/ModelingToolsEditorMode/Private/ModelingToolsEditorMode.cpp
+++ b/Engine/Plugins/Editor/ModelingToolsEditorMode/Source/ModelingToolsEditorMode/Private/ModelingToolsEditorMode.cpp
@@ -377,6 +377,7 @@ void UModelingToolsEditorMode::Enter()
 	if (ModelingModeSettings)
 	{
 		GetInteractiveToolsContext()->SetForceCombinedGizmoMode(ModelingModeSettings->bRespectLevelEditorGizmoMode == false);
+		GetInteractiveToolsContext()->SetAbsoluteWorldSnappingEnabled(ModelingModeSettings->bEnableAbsoluteWorldSnapping);
 	}
 
 	// register snapping manager
diff --git a/Engine/Plugins/Editor/ModelingToolsEditorMode/Source/ModelingToolsEditorMode/Public/ModelingSelectionInteraction.h b/Engine/Plugins/Editor/ModelingToolsEditorMode/Source/ModelingToolsEditorMode/Public/ModelingSelectionInteraction.h
--- a/Engine/Plugins/Editor/ModelingToolsEditorMode/Source/ModelingToolsEditorMode/Public/ModelingSelectionInteraction.h
+++ b/Engine/Plugins/Editor/ModelingToolsEditorMode/Source/ModelingToolsEditorMode/Public/ModelingSelectionInteraction.h
@@ -178,6 +178,9 @@ protected:
 	FVector3d InitialGizmoScale;
 	UE::Geometry::FFrame3d LastUpdateGizmoFrame;
 	FVector3d LastUpdateGizmoScale;
+	FVector3d LastTranslationDelta;
+	FVector4d LastRotateDelta;
+	FVector3d LastScaleDelta;
 	bool bInActiveTransform = false;
 
 	void UpdateGizmoOnSelectionChange();
diff --git a/Engine/Plugins/Editor/ModelingToolsEditorMode/Source/ModelingToolsEditorMode/Public/ModelingToolsEditorModeSettings.h b/Engine/Plugins/Editor/ModelingToolsEditorMode/Source/ModelingToolsEditorMode/Public/ModelingToolsEditorModeSettings.h
--- a/Engine/Plugins/Editor/ModelingToolsEditorMode/Source/ModelingToolsEditorMode/Public/ModelingToolsEditorModeSettings.h
+++ b/Engine/Plugins/Editor/ModelingToolsEditorMode/Source/ModelingToolsEditorMode/Public/ModelingToolsEditorModeSettings.h
@@ -196,6 +196,18 @@ protected:
 
 	bool bRestrictiveMode = false;
 
+
+
+
+public:
+	//
+	// Settings that are currently stored during a single session but are not stored in the config file
+	// (may promote them to persistent settings in the future)
+	//
+
+	/** Toggle Absolute World Grid Position snapping */
+	UPROPERTY(Transient)
+	bool bEnableAbsoluteWorldSnapping = false;
 };
 
 
diff --git a/Engine/Plugins/Runtime/MeshModelingToolset/Source/ModelingComponents/Private/Snapping/ModelingSceneSnappingManager.cpp b/Engine/Plugins/Runtime/MeshModelingToolset/Source/ModelingComponents/Private/Snapping/ModelingSceneSnappingManager.cpp
--- a/Engine/Plugins/Runtime/MeshModelingToolset/Source/ModelingComponents/Private/Snapping/ModelingSceneSnappingManager.cpp
+++ b/Engine/Plugins/Runtime/MeshModelingToolset/Source/ModelingComponents/Private/Snapping/ModelingSceneSnappingManager.cpp
@@ -21,6 +21,8 @@
 #include "RawIndexBuffer.h"
 #include "StaticMeshResources.h"
 #include "UObject/UObjectGlobals.h"
+#include "VectorUtil.h"
+#include "SegmentTypes.h"
 
 #include UE_INLINE_GENERATED_CPP_BY_NAME(ModelingSceneSnappingManager)
 
@@ -29,59 +31,10 @@ using namespace UE::Geometry;
 #define LOCTEXT_NAMESPACE "USceneSnappingManager"
 
 
-
 // defined in SceneGeometrySpatialCache.cpp
 extern TAutoConsoleVariable<bool> CVarEnableModelingVolumeSnapping;
 
 
-static double SnapToIncrement(double fValue, double fIncrement, double offset = 0)
-{
-	if (!FMath::IsFinite(fValue))
-	{
-		return 0;
-	}
-	fValue -= offset;
-	double sign = FMath::Sign(fValue);
-	fValue = FMath::Abs(fValue);
-	int64 nInc = (int64)(fValue / fIncrement);
-	double fRem = (double)fmod(fValue, fIncrement);
-	if (fRem > fIncrement / 2.0)
-	{
-		++nInc;
-	}
-	return sign * (double)nInc * fIncrement + offset;
-}
-
-//@ todo this are mirrored from GeometryProcessing, which is still experimental...replace w/ direct calls once GP component is standardized
-static double OpeningAngleDeg(FVector3d A, FVector3d B, const FVector3d& P)
-{
-	A -= P;
-	A.Normalize();
-	B -= P;
-	B.Normalize();
-	double Dot = FMath::Clamp(FVector3d::DotProduct(A,B), -1.0, 1.0);
-	return FMathd::ACos(Dot) * (180.0 / 3.141592653589);
-}
-
-static FVector3d NearestSegmentPt(FVector3d A, FVector3d B, const FVector3d& P)
-{
-	FVector3d Direction = (B - A);
-	double Length = Direction.Size();
-	Direction /= Length;
-	double t = FVector3d::DotProduct( (P - A), Direction);
-	if (t >= Length)
-	{
-		return B;
-	}
-	if (t <= 0)
-	{
-		return A;
-	}
-	return A + t * Direction;
-}
-
-
-
 
 void UModelingSceneSnappingManager::Initialize(TObjectPtr<UInteractiveToolsContext> ToolsContext)
 {
@@ -524,9 +477,9 @@ bool UModelingSceneSnappingManager::ExecuteSceneSnapQueryPosition(const FSceneSn
 
 			FVector GridSize = Request.GridSize.Get(SnappingConfig.PositionGridDimensions);
 
-			SnapResult.Position.X = (FVector::FReal)SnapToIncrement(Request.Position.X, GridSize.X);
-			SnapResult.Position.Y = (FVector::FReal)SnapToIncrement(Request.Position.Y, GridSize.Y);
-			SnapResult.Position.Z = (FVector::FReal)SnapToIncrement(Request.Position.Z, GridSize.Z);
+			SnapResult.Position.X = (FVector::FReal)UE::Geometry::SnapToIncrement(Request.Position.X, GridSize.X);
+			SnapResult.Position.Y = (FVector::FReal)UE::Geometry::SnapToIncrement(Request.Position.Y, GridSize.Y);
+			SnapResult.Position.Z = (FVector::FReal)UE::Geometry::SnapToIncrement(Request.Position.Z, GridSize.Z);
 
 			Results.Add(SnapResult);
 			FoundResultCount++;
@@ -549,7 +502,7 @@ bool UModelingSceneSnappingManager::ExecuteSceneSnapQueryPosition(const FSceneSn
 		{
 			for (int j = 0; j < 3; ++j)
 			{
-				double VisualAngle = OpeningAngleDeg((FVector3d)RequestIn.Position, (FVector3d)SnapResult.TriVertices[j], WorldRay.Origin);
+				double VisualAngle = UE::Geometry::VectorUtil::OpeningAngleD((FVector3d)RequestIn.Position, (FVector3d)SnapResult.TriVertices[j], WorldRay.Origin);
 				if (VisualAngle < SmallestSnapAngle)
 				{
 					SmallestSnapAngle = VisualAngle;
@@ -569,8 +522,9 @@ bool UModelingSceneSnappingManager::ExecuteSceneSnapQueryPosition(const FSceneSn
 		{
 			for (int j = 0; j < 3; ++j)
 			{
-				FVector3d EdgeNearestPt = NearestSegmentPt((FVector3d)SnapResult.TriVertices[j], (FVector3d)SnapResult.TriVertices[(j+1)%3], (FVector3d)RequestIn.Position);
-				double VisualAngle = OpeningAngleDeg((FVector3d)RequestIn.Position, EdgeNearestPt, WorldRay.Origin);
+				UE::Geometry::FSegment3d Segment(SnapResult.TriVertices[j], SnapResult.TriVertices[(j+1)%3]);
+				FVector3d EdgeNearestPt = Segment.NearestPoint(RequestIn.Position);
+				double VisualAngle = UE::Geometry::VectorUtil::OpeningAngleD(RequestIn.Position, EdgeNearestPt, WorldRay.Origin);
 				if (VisualAngle < SmallestSnapAngle )
 				{
 					SmallestSnapAngle = VisualAngle;
@@ -622,7 +576,7 @@ bool UModelingSceneSnappingManager::ExecuteSceneSnapQueryPosition(const FSceneSn
 	if (bHitSceneCache)
 	{
 		// is this check necessary? isn't this always going to be zero for a scene hit??
-		double VisualAngle = OpeningAngleDeg((FVector3d)Request.Position, HitPoint.WorldPoint, RayStart);
+		double VisualAngle = UE::Geometry::VectorUtil::OpeningAngleD((FVector3d)Request.Position, HitPoint.WorldPoint, RayStart);
 		if (VisualAngle < (double)Request.VisualAngleThresholdDegrees)
 		{
 			FVector3d A, B, C;
@@ -657,7 +611,7 @@ bool UModelingSceneSnappingManager::ExecuteSceneSnapQueryPosition(const FSceneSn
 	if (bHitWorld)
 	{
 		// is this check necessary? isn't this always going to be zero for a scene hit??
-		double VisualAngle = OpeningAngleDeg((FVector3d)Request.Position, (FVector3d)HitResult.ImpactPoint, RayStart);
+		double VisualAngle = UE::Geometry::VectorUtil::OpeningAngleD((FVector3d)Request.Position, (FVector3d)HitResult.ImpactPoint, RayStart);
 		if (VisualAngle < (double)Request.VisualAngleThresholdDegrees)
 		{
 			FSceneSnapQueryResult SnapResult;
diff --git a/Engine/Source/Editor/Experimental/EditorInteractiveToolsFramework/Private/EdModeInteractiveToolsContext.cpp b/Engine/Source/Editor/Experimental/EditorInteractiveToolsFramework/Private/EdModeInteractiveToolsContext.cpp
--- a/Engine/Source/Editor/Experimental/EditorInteractiveToolsFramework/Private/EdModeInteractiveToolsContext.cpp
+++ b/Engine/Source/Editor/Experimental/EditorInteractiveToolsFramework/Private/EdModeInteractiveToolsContext.cpp
@@ -168,6 +168,7 @@ public:
 		Config.PositionGridDimensions = FVector(EditorGridSize, EditorGridSize, EditorGridSize);
 		Config.bEnableRotationGridSnapping = (GetDefault<ULevelEditorViewportSettings>()->RotGridEnabled != 0);
 		Config.RotationGridAngles = GEditor->GetRotGridSize();
+		Config.bEnableAbsoluteWorldSnapping = ToolsContext->GetAbsoluteWorldSnappingEnabled();
 		return Config;
 	}
 
@@ -841,6 +842,11 @@ void UEditorInteractiveToolsContext::SetForceCombinedGizmoMode(bool bEnabled)
 	bForceCombinedGizmoMode = bEnabled;
 }
 
+void UEditorInteractiveToolsContext::SetAbsoluteWorldSnappingEnabled(bool bEnabled)
+{
+	bEnableAbsoluteWorldSnapping = bEnabled;
+}
+
 void UEditorInteractiveToolsContext::OnLevelEditorCreated(TSharedPtr<ILevelEditor> InLevelEditor)
 {
 	if (UTypedElementSelectionSet* TypedElementSelectionSet = EditorModeManager->GetEditorSelectionSet())
diff --git a/Engine/Source/Editor/Experimental/EditorInteractiveToolsFramework/Public/EdModeInteractiveToolsContext.h b/Engine/Source/Editor/Experimental/EditorInteractiveToolsFramework/Public/EdModeInteractiveToolsContext.h
--- a/Engine/Source/Editor/Experimental/EditorInteractiveToolsFramework/Public/EdModeInteractiveToolsContext.h
+++ b/Engine/Source/Editor/Experimental/EditorInteractiveToolsFramework/Public/EdModeInteractiveToolsContext.h
@@ -159,6 +159,16 @@ public:
 	bool GetForceCombinedGizmoModeEnabled() const { return bForceCombinedGizmoMode; }
 
 
+	/**
+	 * Configure whether Transform Gizmos created by the ITF (eg CombinedTransformGizmo) should, when in World coordinate system,
+	 * snap to an Absolute world-aligned grid, or snap Relative to the initial position of any particular gizmo transform.
+	 * Relative is the default and is also the behavior of the standard UE Gizmo.
+	 */
+	void SetAbsoluteWorldSnappingEnabled(bool bEnabled);
+
+	/** @return true if Absolute World Snapping mode is Enabled */
+	bool GetAbsoluteWorldSnappingEnabled() const { return bEnableAbsoluteWorldSnapping; }
+
 protected:
 	/**
 	 * Called upon Level Editor Created in order to bind to EditorElementSelectionPtr changed
@@ -225,6 +235,7 @@ private:
 	bool bEnableRenderingDuringHitProxyPass = true;
 
 	bool bForceCombinedGizmoMode = false;
+	bool bEnableAbsoluteWorldSnapping = false;
 
 	bool bIsActive = false;
 };
diff --git a/Engine/Source/Runtime/GeometryCore/Public/Quaternion.h b/Engine/Source/Runtime/GeometryCore/Public/Quaternion.h
--- a/Engine/Source/Runtime/GeometryCore/Public/Quaternion.h
+++ b/Engine/Source/Runtime/GeometryCore/Public/Quaternion.h
@@ -142,7 +142,10 @@ PRAGMA_ENABLE_DEPRECATION_WARNINGS
 		Z = (RealType)Quat.Z;
 		W = (RealType)Quat.W;
 	}
-
+	explicit inline operator UE::Math::TVector4<RealType>() const
+	{
+		return UE::Math::TVector4<RealType>(X,Y,Z,W);
+	}
 };
 
 
diff --git a/Engine/Source/Runtime/GeometryCore/Public/VectorUtil.h b/Engine/Source/Runtime/GeometryCore/Public/VectorUtil.h
--- a/Engine/Source/Runtime/GeometryCore/Public/VectorUtil.h
+++ b/Engine/Source/Runtime/GeometryCore/Public/VectorUtil.h
@@ -531,5 +531,31 @@ namespace VectorUtil
 }; // namespace VectorUtil
 
 
+
+/**
+ * Snap Value to steps of Increment (centered at 0, ie steps are -Increment, 0, Increment, 2*Increment, ...).
+ * Optional Offset can be used to snap relative value.
+ */
+template<typename RealType>
+RealType SnapToIncrement(RealType Value, RealType Increment, RealType Offset = 0)
+{
+	if (!FMath::IsFinite(Value))
+	{
+		return (RealType)0;
+	}
+	Value -= Offset;
+	RealType ValueSign = FMath::Sign(Value);
+	Value = FMath::Abs(Value);
+	int64 IntegerIncrement = (int64)(Value / Increment);
+	RealType Remainder = (RealType)fmod(Value, Increment);
+	if (Remainder > Increment / 2.0)
+	{
+		++IntegerIncrement;
+	}
+	return ValueSign * (RealType)IntegerIncrement * Increment + Offset;
+}
+
+
+
 }  // end namespace Geometry
 }  // end namespace UE
diff --git a/Engine/Source/Runtime/InteractiveToolsFramework/Private/BaseGizmos/CombinedTransformGizmo.cpp b/Engine/Source/Runtime/InteractiveToolsFramework/Private/BaseGizmos/CombinedTransformGizmo.cpp
--- a/Engine/Source/Runtime/InteractiveToolsFramework/Private/BaseGizmos/CombinedTransformGizmo.cpp
+++ b/Engine/Source/Runtime/InteractiveToolsFramework/Private/BaseGizmos/CombinedTransformGizmo.cpp
@@ -15,6 +15,9 @@
 #include "BaseGizmos/GizmoLineHandleComponent.h"
 #include "BaseGizmos/GizmoViewContext.h"
 
+#include "MathUtil.h"
+#include "VectorUtil.h"
+
 // need this to implement hover
 #include "BaseGizmos/GizmoBaseComponent.h"
 
@@ -426,10 +429,12 @@ void UCombinedTransformGizmo::Tick(float DeltaTime)
 	{
 		CurrentCoordinateSystem = GetGizmoManager()->GetContextQueriesAPI()->GetCurrentCoordinateSystem();
 	}
-	
 	check(CurrentCoordinateSystem == EToolContextCoordinateSystem::World || CurrentCoordinateSystem == EToolContextCoordinateSystem::Local)
-	bool bUseLocalAxes = (CurrentCoordinateSystem == EToolContextCoordinateSystem::Local);
+	
+	FToolContextSnappingConfiguration SnappingConfig = GetGizmoManager()->GetContextQueriesAPI()->GetCurrentSnappingSettings();
+	RelativeTranslationSnapping.UpdateContextValue(SnappingConfig.bEnableAbsoluteWorldSnapping == false);
 
+	bool bUseLocalAxes = (CurrentCoordinateSystem == EToolContextCoordinateSystem::Local);
 	if (AxisXSource != nullptr && AxisYSource != nullptr && AxisZSource != nullptr)
 	{
 		AxisXSource->bLocalAxes = bUseLocalAxes;
@@ -524,19 +529,19 @@ void UCombinedTransformGizmo::SetActiveTarget(UTransformProxy* Target, IToolCont
 	// todo should we hold onto these?
 	if (GizmoActor->TranslateX != nullptr)
 	{
-		UInteractiveGizmo* NewGizmo = AddAxisTranslationGizmo(GizmoActor->TranslateX, GizmoComponent, AxisXSource, TransformSource, StateTarget);
+		UInteractiveGizmo* NewGizmo = AddAxisTranslationGizmo(GizmoActor->TranslateX, GizmoComponent, AxisXSource, TransformSource, StateTarget, 0);
 		ActiveComponents.Add(GizmoActor->TranslateX);
 		TranslationSubGizmos.Add( FSubGizmoInfo{ GizmoActor->TranslateX, NewGizmo } );
 	}
 	if (GizmoActor->TranslateY != nullptr)
 	{
-		UInteractiveGizmo* NewGizmo = AddAxisTranslationGizmo(GizmoActor->TranslateY, GizmoComponent, AxisYSource, TransformSource, StateTarget);
+		UInteractiveGizmo* NewGizmo = AddAxisTranslationGizmo(GizmoActor->TranslateY, GizmoComponent, AxisYSource, TransformSource, StateTarget, 1);
 		ActiveComponents.Add(GizmoActor->TranslateY);
 		TranslationSubGizmos.Add(FSubGizmoInfo{ GizmoActor->TranslateY, NewGizmo });
 	}
 	if (GizmoActor->TranslateZ != nullptr)
 	{
-		UInteractiveGizmo* NewGizmo = AddAxisTranslationGizmo(GizmoActor->TranslateZ, GizmoComponent, AxisZSource, TransformSource, StateTarget);
+		UInteractiveGizmo* NewGizmo = AddAxisTranslationGizmo(GizmoActor->TranslateZ, GizmoComponent, AxisZSource, TransformSource, StateTarget, 2);
 		ActiveComponents.Add(GizmoActor->TranslateZ);
 		TranslationSubGizmos.Add(FSubGizmoInfo{ GizmoActor->TranslateZ, NewGizmo });
 	}
@@ -544,19 +549,19 @@ void UCombinedTransformGizmo::SetActiveTarget(UTransformProxy* Target, IToolCont
 
 	if (GizmoActor->TranslateYZ != nullptr)
 	{
-		UInteractiveGizmo* NewGizmo = AddPlaneTranslationGizmo(GizmoActor->TranslateYZ, GizmoComponent, AxisXSource, TransformSource, StateTarget);
+		UInteractiveGizmo* NewGizmo = AddPlaneTranslationGizmo(GizmoActor->TranslateYZ, GizmoComponent, AxisXSource, TransformSource, StateTarget, 1, 2);
 		ActiveComponents.Add(GizmoActor->TranslateYZ);
 		TranslationSubGizmos.Add(FSubGizmoInfo{ GizmoActor->TranslateYZ, NewGizmo });
 	}
 	if (GizmoActor->TranslateXZ != nullptr)
 	{
-		UInteractiveGizmo* NewGizmo = AddPlaneTranslationGizmo(GizmoActor->TranslateXZ, GizmoComponent, AxisYSource, TransformSource, StateTarget);
+		UInteractiveGizmo* NewGizmo = AddPlaneTranslationGizmo(GizmoActor->TranslateXZ, GizmoComponent, AxisYSource, TransformSource, StateTarget, 2, 0);	// flip here corresponds to UGizmoComponentAxisSource::GetTangentVectors()
 		ActiveComponents.Add(GizmoActor->TranslateXZ);
 		TranslationSubGizmos.Add(FSubGizmoInfo{ GizmoActor->TranslateXZ, NewGizmo });
 	}
 	if (GizmoActor->TranslateXY != nullptr)
 	{
-		UInteractiveGizmo* NewGizmo = AddPlaneTranslationGizmo(GizmoActor->TranslateXY, GizmoComponent, AxisZSource, TransformSource, StateTarget);
+		UInteractiveGizmo* NewGizmo = AddPlaneTranslationGizmo(GizmoActor->TranslateXY, GizmoComponent, AxisZSource, TransformSource, StateTarget, 0, 1);
 		ActiveComponents.Add(GizmoActor->TranslateXY);
 		TranslationSubGizmos.Add(FSubGizmoInfo{ GizmoActor->TranslateXY, NewGizmo });
 	}
@@ -713,7 +718,8 @@ UInteractiveGizmo* UCombinedTransformGizmo::AddAxisTranslationGizmo(
 	UPrimitiveComponent* AxisComponent, USceneComponent* RootComponent,
 	IGizmoAxisSource* AxisSource,
 	IGizmoTransformSource* TransformSource,
-	IGizmoStateTarget* StateTargetIn)
+	IGizmoStateTarget* StateTargetIn,
+	int AxisIndex)
 {
 	// create axis-position gizmo, axis-position parameter will drive translation
 	UAxisPositionGizmo* TranslateGizmo = Cast<UAxisPositionGizmo>(GetGizmoManager()->CreateGizmo(AxisPositionBuilderIdentifier));
@@ -725,6 +731,7 @@ UInteractiveGizmo* UCombinedTransformGizmo::AddAxisTranslationGizmo(
 	// parameter source maps axis-parameter-change to translation of TransformSource's transform
 	UGizmoAxisTranslationParameterSource* ParamSource = UGizmoAxisTranslationParameterSource::Construct(AxisSource, TransformSource, this);
 	ParamSource->PositionConstraintFunction = [this](const FVector& Pos, FVector& Snapped) { return PositionSnapFunction(Pos, Snapped); };
+	ParamSource->AxisDeltaConstraintFunction = [this, AxisIndex](double AxisDelta, double& SnappedAxisDelta) { return PositionAxisDeltaSnapFunction(AxisDelta, SnappedAxisDelta, AxisIndex); };
 	TranslateGizmo->ParameterSource = ParamSource;
 
 	// sub-component provides hit target
@@ -753,7 +760,8 @@ UInteractiveGizmo* UCombinedTransformGizmo::AddPlaneTranslationGizmo(
 	UPrimitiveComponent* AxisComponent, USceneComponent* RootComponent,
 	IGizmoAxisSource* AxisSource,
 	IGizmoTransformSource* TransformSource,
-	IGizmoStateTarget* StateTargetIn)
+	IGizmoStateTarget* StateTargetIn,
+	int XAxisIndex, int YAxisIndex)
 {
 	// create axis-position gizmo, axis-position parameter will drive translation
 	UPlanePositionGizmo* TranslateGizmo = Cast<UPlanePositionGizmo>(GetGizmoManager()->CreateGizmo(PlanePositionBuilderIdentifier));
@@ -765,6 +773,8 @@ UInteractiveGizmo* UCombinedTransformGizmo::AddPlaneTranslationGizmo(
 	// parameter source maps axis-parameter-change to translation of TransformSource's transform
 	UGizmoPlaneTranslationParameterSource* ParamSource = UGizmoPlaneTranslationParameterSource::Construct(AxisSource, TransformSource, this);
 	ParamSource->PositionConstraintFunction = [this](const FVector& Pos, FVector& Snapped) { return PositionSnapFunction(Pos, Snapped); };
+	ParamSource->AxisXDeltaConstraintFunction = [this, XAxisIndex](double AxisDelta, double& SnappedAxisDelta) { return PositionAxisDeltaSnapFunction(AxisDelta, SnappedAxisDelta, XAxisIndex); };
+	ParamSource->AxisYDeltaConstraintFunction = [this, YAxisIndex](double AxisDelta, double& SnappedAxisDelta) { return PositionAxisDeltaSnapFunction(AxisDelta, SnappedAxisDelta, YAxisIndex); };
 	TranslateGizmo->ParameterSource = ParamSource;
 
 	// sub-component provides hit target
@@ -806,7 +816,9 @@ UInteractiveGizmo* UCombinedTransformGizmo::AddAxisRotationGizmo(
 
 	// parameter source maps angle-parameter-change to rotation of TransformSource's transform
 	UGizmoAxisRotationParameterSource* AngleSource = UGizmoAxisRotationParameterSource::Construct(AxisSource, TransformSource, this);
-	AngleSource->RotationConstraintFunction = [this](const FQuat& DeltaRotation){ return RotationSnapFunction(DeltaRotation); };
+	// axis rotation is currently only relative so it should only ever snap angle-deltas
+	//AngleSource->RotationConstraintFunction = [this](const FQuat& DeltaRotation){ return RotationSnapFunction(DeltaRotation); };
+	AngleSource->AngleDeltaConstraintFunction = [this](double AngleDelta, double& SnappedDelta){ return RotationAxisAngleSnapFunction(AngleDelta, SnappedDelta, 0); };
 	RotateGizmo->AngleSource = AngleSource;
 
 	// sub-component provides hit target
@@ -971,84 +983,149 @@ bool UCombinedTransformGizmo::PositionSnapFunction(const FVector& WorldPosition,
 {
 	SnappedPositionOut = WorldPosition;
 
-	// only snap if we want snapping obvs
-	if (bSnapToWorldGrid == false)
+	// only snap world positions if we want world position snapping...
+	if (bSnapToWorldGrid == false || RelativeTranslationSnapping.IsEnabled() == true)
 	{
 		return false;
 	}
 
-	// only snap to world grid when using world axes
+	// we can only snap positions in world coordinate system
 	EToolContextCoordinateSystem CoordSystem = GetGizmoManager()->GetContextQueriesAPI()->GetCurrentCoordinateSystem();
 	if (CoordSystem != EToolContextCoordinateSystem::World)
 	{
 		return false;
 	}
 
-	USceneSnappingManager* SnapManager = USceneSnappingManager::Find(GetGizmoManager());
-	if (!SnapManager)
+	// need a snapping manager
+	if ( USceneSnappingManager* SnapManager = USceneSnappingManager::Find(GetGizmoManager()) )
 	{
-		return false;
+		FSceneSnapQueryRequest Request;
+		Request.RequestType = ESceneSnapQueryType::Position;
+		Request.TargetTypes = ESceneSnapQueryTargetType::Grid;
+		if ( bGridSizeIsExplicit )
+		{
+			Request.GridSize = ExplicitGridSize;
+		}
+		TArray<FSceneSnapQueryResult> Results;
+		Results.Reserve(1);
+
+		Request.Position = WorldPosition;
+		if (SnapManager->ExecuteSceneSnapQuery(Request, Results))
+		{
+			SnappedPositionOut = Results[0].Position;
+			return true;
+		};
 	}
 
-	FSceneSnapQueryRequest Request;
-	Request.RequestType = ESceneSnapQueryType::Position;
-	Request.TargetTypes = ESceneSnapQueryTargetType::Grid;
-	Request.Position = WorldPosition;
-	if ( bGridSizeIsExplicit )
+	return false;
+}
+
+
+bool UCombinedTransformGizmo::PositionAxisDeltaSnapFunction(double AxisDelta, double& SnappedDeltaOut, int AxisIndex) const
+{
+	if (!bSnapToWorldGrid) return false;
+
+	EToolContextCoordinateSystem CoordSystem = GetGizmoManager()->GetContextQueriesAPI()->GetCurrentCoordinateSystem();
+	bool bUseRelativeSnapping = RelativeTranslationSnapping.IsEnabled() || (CoordSystem != EToolContextCoordinateSystem::World);
+	if (!bUseRelativeSnapping)
 	{
-		Request.GridSize = ExplicitGridSize;
+		return false;
 	}
-	TArray<FSceneSnapQueryResult> Results;
-	if (SnapManager->ExecuteSceneSnapQuery(Request, Results))
-	{
-		SnappedPositionOut = Results[0].Position;
-		return true;
-	};
 
+	if ( USceneSnappingManager* SnapManager = USceneSnappingManager::Find(GetGizmoManager()) )
+	{
+		FSceneSnapQueryRequest Request;
+		Request.RequestType = ESceneSnapQueryType::Position;
+		Request.TargetTypes = ESceneSnapQueryTargetType::Grid;
+		if ( bGridSizeIsExplicit )
+		{
+			Request.GridSize = ExplicitGridSize;
+		}
+		TArray<FSceneSnapQueryResult> Results;
+		Results.Reserve(1);
+		
+		// this is a bit of a hack, since the snap query only snaps world points, and the grid may not be
+		// uniform. A point on the specified X/Y/Z at the delta-distance is snapped, this is ideally
+		// equivalent to actually computing a snap of the axis-delta
+		Request.Position = FVector::Zero();
+		Request.Position[AxisIndex] = AxisDelta;
+		if (SnapManager->ExecuteSceneSnapQuery(Request, Results))
+		{
+			SnappedDeltaOut = Results[0].Position[AxisIndex];
+			return true;
+		};
+	}
 	return false;
 }
 
+
+
+
 FQuat UCombinedTransformGizmo::RotationSnapFunction(const FQuat& DeltaRotation) const
 {
+	// note: this is currently unused. Although we can snap to the "rotation grid", since the
+	// gizmo only supports axis rotations, it doesn't make sense. Leaving in for now in case
+	// a "tumble" handle is added, in which case it makes sense to snap to the world rotation grid...
+
 	FQuat SnappedDeltaRotation = DeltaRotation;
 
-	// only snap if we want snapping obvs
-	if (!bSnapToWorldGrid)
+	// only snap world positions if we want world position snapping...
+	if (bSnapToWorldRotGrid == false )
 	{
 		return SnappedDeltaRotation;
 	}
 
-	// To match our position snapping behavior, only snap when using world axes.
-	// Note that if we someday want to snap in local mode, this function will need further
-	// changing because the quaternion is given and snapped in world space, whereas we
-	// would want to snap it relative to the local frame start orientation.
+	// can only snap absolute rotations in World coordinates
 	EToolContextCoordinateSystem CoordSystem = GetGizmoManager()->GetContextQueriesAPI()->GetCurrentCoordinateSystem();
 	if (CoordSystem != EToolContextCoordinateSystem::World)
 	{
 		return SnappedDeltaRotation;
 	}
 
-	USceneSnappingManager* SnapManager = USceneSnappingManager::Find(GetGizmoManager());
-	if (!SnapManager)
+	// need a snapping manager
+	if ( USceneSnappingManager* SnapManager = USceneSnappingManager::Find(GetGizmoManager()) )
 	{
-		return SnappedDeltaRotation;
+		FSceneSnapQueryRequest Request;
+		Request.RequestType   = ESceneSnapQueryType::Rotation;
+		Request.TargetTypes   = ESceneSnapQueryTargetType::Grid;
+		Request.DeltaRotation = DeltaRotation;
+		if ( bRotationGridSizeIsExplicit )
+		{
+			Request.RotGridSize = ExplicitRotationGridSize;
+		}
+		TArray<FSceneSnapQueryResult> Results;
+		if (SnapManager->ExecuteSceneSnapQuery(Request, Results))
+		{
+			SnappedDeltaRotation = Results[0].DeltaRotation;
+		};
 	}
 
-	FSceneSnapQueryRequest Request;
-	Request.RequestType   = ESceneSnapQueryType::Rotation;
-	Request.TargetTypes   = ESceneSnapQueryTargetType::Grid;
-	Request.DeltaRotation = DeltaRotation;
-	if ( bRotationGridSizeIsExplicit )
+	return SnappedDeltaRotation;
+}
+
+
+
+
+
+bool UCombinedTransformGizmo::RotationAxisAngleSnapFunction(double AxisAngleDelta, double& SnappedAxisAngleDeltaOut, int AxisIndex) const
+{
+	if (!bSnapToWorldRotGrid) return false;
+
+	FToolContextSnappingConfiguration SnappingConfig = GetGizmoManager()->GetContextQueriesAPI()->GetCurrentSnappingSettings();
+	if ( SnappingConfig.bEnableRotationGridSnapping )
 	{
-		Request.RotGridSize = ExplicitRotationGridSize;
+		double SnapDelta = SnappingConfig.RotationGridAngles.Yaw;		// could use AxisIndex here?
+		if ( bRotationGridSizeIsExplicit )
+		{
+			SnapDelta = ExplicitRotationGridSize.Yaw;
+		}
+		AxisAngleDelta *= FMathd::RadToDeg;
+		SnappedAxisAngleDeltaOut = UE::Geometry::SnapToIncrement(AxisAngleDelta, SnapDelta);
+		SnappedAxisAngleDeltaOut *= FMathd::DegToRad;
+		return true;
 	}
-	TArray<FSceneSnapQueryResult> Results;
-	if (SnapManager->ExecuteSceneSnapQuery(Request, Results))
-	{
-		SnappedDeltaRotation = Results[0].DeltaRotation;
-	};
 
-	return SnappedDeltaRotation;
+	return false;
 }
 
 #undef LOCTEXT_NAMESPACE
diff --git a/Engine/Source/Runtime/InteractiveToolsFramework/Private/BaseGizmos/ParameterToTransformAdapters.cpp b/Engine/Source/Runtime/InteractiveToolsFramework/Private/BaseGizmos/ParameterToTransformAdapters.cpp
new file mode 100644
--- /dev/null
+++ b/Engine/Source/Runtime/InteractiveToolsFramework/Private/BaseGizmos/ParameterToTransformAdapters.cpp
@@ -0,0 +1,329 @@
+// Copyright Epic Games, Inc. All Rights Reserved.
+ 
+#include "BaseGizmos/ParameterToTransformAdapters.h"
+#include "BaseGizmos/GizmoMath.h"
+
+
+void UGizmoAxisTranslationParameterSource::SetParameter(float NewValue)
+{
+	Parameter = NewValue;
+	LastChange.CurrentValue = NewValue;
+
+	double UseDelta = LastChange.GetChangeDelta();
+
+	// check for any constraints on the delta value
+	double SnappedDelta = 0;
+	if (AxisDeltaConstraintFunction(UseDelta, SnappedDelta))
+	{
+		UseDelta = SnappedDelta;
+	}
+
+	// construct translation as delta from initial position
+	FVector Translation = UseDelta * CurTranslationAxis;
+
+	// translate the initial transform
+	FTransform NewTransform = InitialTransform;
+	NewTransform.AddToTranslation(Translation);
+
+	// apply position constraint
+	FVector SnappedPos;
+	if (PositionConstraintFunction(NewTransform.GetTranslation(), SnappedPos))
+	{
+		FVector SnappedLinePos = GizmoMath::ProjectPointOntoLine(SnappedPos, CurTranslationOrigin, CurTranslationAxis);
+		NewTransform.SetTranslation(SnappedLinePos);
+	}
+
+	TransformSource->SetTransform(NewTransform);
+
+	OnParameterChanged.Broadcast(this, LastChange);
+}
+
+void UGizmoAxisTranslationParameterSource::BeginModify()
+{
+	check(AxisSource);
+
+	LastChange = FGizmoFloatParameterChange(Parameter);
+
+	InitialTransform = TransformSource->GetTransform();
+	CurTranslationAxis = AxisSource->GetDirection();
+	CurTranslationOrigin = AxisSource->GetOrigin();
+}
+
+void UGizmoAxisTranslationParameterSource::EndModify()
+{
+}
+
+
+
+
+
+
+
+void UGizmoPlaneTranslationParameterSource::SetParameter(const FVector2D& NewValue)
+{
+	Parameter = NewValue;
+	LastChange.CurrentValue = NewValue;
+
+	// construct translation as delta from initial position
+	FVector2D Delta = LastChange.GetChangeDelta();
+	double UseDeltaX = Delta.X;
+	double UseDeltaY = Delta.Y;
+
+	// check for any constraints on the delta value
+	double SnappedDeltaX = 0, SnappedDeltaY = 0;
+	if (AxisXDeltaConstraintFunction(UseDeltaX, SnappedDeltaX))
+	{
+		UseDeltaX = SnappedDeltaX;
+	}
+	if (AxisYDeltaConstraintFunction(UseDeltaY, SnappedDeltaY))
+	{
+		UseDeltaY = SnappedDeltaY;
+	}
+
+	FVector Translation = UseDeltaX*CurTranslationAxisX + UseDeltaY*CurTranslationAxisY;
+
+	// apply translation to initial transform
+	FTransform NewTransform = InitialTransform;
+	NewTransform.AddToTranslation(Translation);
+
+	// apply position constraint
+	FVector SnappedPos;
+	if (PositionConstraintFunction(NewTransform.GetTranslation(), SnappedPos))
+	{
+		FVector PlanePos = GizmoMath::ProjectPointOntoPlane(SnappedPos, CurTranslationOrigin, CurTranslationNormal);
+		NewTransform.SetTranslation(PlanePos);
+	}
+
+	TransformSource->SetTransform(NewTransform);
+
+	OnParameterChanged.Broadcast(this, LastChange);
+}
+
+void UGizmoPlaneTranslationParameterSource::BeginModify()
+{
+	check(AxisSource);
+
+	LastChange = FGizmoVec2ParameterChange(Parameter);
+
+	// save initial transformation and axis information
+	InitialTransform = TransformSource->GetTransform();
+	CurTranslationOrigin = AxisSource->GetOrigin();
+	AxisSource->GetAxisFrame(CurTranslationNormal, CurTranslationAxisX, CurTranslationAxisY);
+}
+
+void UGizmoPlaneTranslationParameterSource::EndModify()
+{
+}
+
+
+
+void UGizmoAxisRotationParameterSource::SetParameter(float NewValue)
+{
+	Angle = NewValue;
+	LastChange.CurrentValue = NewValue;
+
+	double AngleDelta = LastChange.GetChangeDelta();
+	double SnappedDelta;
+	if (AngleDeltaConstraintFunction(AngleDelta, SnappedDelta))
+	{
+		AngleDelta = SnappedDelta;
+	}
+
+	// construct rotation as delta from initial position
+	FQuat DeltaRotation(CurRotationAxis, AngleDelta);
+	DeltaRotation = RotationConstraintFunction(DeltaRotation);
+
+	// rotate the vector from the rotation origin to the transform origin, 
+	// to get the translation of the origin produced by the rotation
+	FVector DeltaPosition = InitialTransform.GetLocation() - CurRotationOrigin;
+	DeltaPosition = DeltaRotation * DeltaPosition;
+	FVector NewLocation = CurRotationOrigin + DeltaPosition;
+
+	// rotate the initial transform by the rotation
+	FQuat NewRotation = DeltaRotation * InitialTransform.GetRotation();
+
+	// construct new transform
+	FTransform NewTransform = InitialTransform;
+	NewTransform.SetLocation(NewLocation);
+	NewTransform.SetRotation(NewRotation);
+	TransformSource->SetTransform(NewTransform);
+
+	OnParameterChanged.Broadcast(this, LastChange);
+}
+
+
+void UGizmoAxisRotationParameterSource::BeginModify()
+{
+	check(AxisSource != nullptr);
+
+	LastChange = FGizmoFloatParameterChange(Angle);
+
+	// save initial transformation and axis information
+	InitialTransform = TransformSource->GetTransform();
+	CurRotationAxis = AxisSource->GetDirection();
+	CurRotationOrigin = AxisSource->GetOrigin();
+}
+
+void UGizmoAxisRotationParameterSource::EndModify()
+{
+}
+
+
+
+
+void UGizmoUniformScaleParameterSource::SetParameter(const FVector2D& NewValue)
+{
+	Parameter = NewValue;
+	LastChange.CurrentValue = NewValue;
+
+	// Convert 2D parameter delta to a 1D uniform scale change
+	// This possibly be exposed as a TFunction to allow customization?
+	double SignedDelta = LastChange.GetChangeDelta().X + LastChange.GetChangeDelta().Y;
+	SignedDelta *= ScaleMultiplier;
+	SignedDelta += 1.0;
+
+	FTransform NewTransform = InitialTransform;
+	FVector StartScale = InitialTransform.GetScale3D();
+	FVector NewScale = SignedDelta * StartScale;
+	NewTransform.SetScale3D(NewScale);
+
+	// apply position constraint
+	//FVector SnappedPos;
+	//if (PositionConstraintFunction(NewTransform.GetScale(), SnappedPos))
+	//{
+	//	FVector PlanePos = GizmoMath::ProjectPointOntoPlane(SnappedPos, CurTranslationOrigin, CurTranslationNormal);
+	//	NewTransform.SetTranslation(PlanePos);
+	//}
+
+	TransformSource->SetTransform(NewTransform);
+
+	OnParameterChanged.Broadcast(this, LastChange);
+}
+
+void UGizmoUniformScaleParameterSource::BeginModify()
+{
+	check(AxisSource);
+
+	LastChange = FGizmoVec2ParameterChange(Parameter);
+
+	// save initial transformation and axis information
+	InitialTransform = TransformSource->GetTransform();
+	CurScaleOrigin = AxisSource->GetOrigin();
+	// note: currently not used!
+	AxisSource->GetAxisFrame(CurScaleNormal, CurScaleAxisX, CurScaleAxisY);
+}
+
+void UGizmoUniformScaleParameterSource::EndModify()
+{
+}
+
+
+
+
+
+void UGizmoAxisScaleParameterSource::SetParameter(float NewValue)
+{
+	Parameter = NewValue;
+	LastChange.CurrentValue = NewValue;
+
+	// construct translation as delta from initial position
+	float ScaleDelta = LastChange.GetChangeDelta() * ScaleMultiplier;
+
+	// translate the initial transform
+	FTransform NewTransform = InitialTransform;
+	FVector StartScale = InitialTransform.GetScale3D();
+	FVector NewScale = StartScale + ScaleDelta * CurScaleAxis;
+
+	if (bClampToZero)
+	{
+		NewScale = FVector::Max(FVector::ZeroVector, NewScale);
+	}
+
+	NewTransform.SetScale3D(NewScale);
+
+
+	// apply position constraint
+	//FVector SnappedPos;
+	//if (ScaleConstraintFunction(NewTransform.GetTranslation(), SnappedPos))
+	//{
+	//	FVector SnappedLinePos = GizmoMath::ProjectPointOntoLine(SnappedPos, CurTranslationOrigin, CurTranslationAxis);
+	//	NewTransform.SetTranslation(SnappedLinePos);
+	//}
+
+	TransformSource->SetTransform(NewTransform);
+
+	OnParameterChanged.Broadcast(this, LastChange);
+}
+
+void UGizmoAxisScaleParameterSource::BeginModify()
+{
+	check(AxisSource);
+
+	LastChange = FGizmoFloatParameterChange(Parameter);
+
+	InitialTransform = TransformSource->GetTransform();
+		
+	CurScaleAxis = AxisSource->GetDirection();
+	CurScaleOrigin = AxisSource->GetOrigin();
+}
+
+void UGizmoAxisScaleParameterSource::EndModify()
+{
+}
+
+
+
+
+
+void UGizmoPlaneScaleParameterSource::SetParameter(const FVector2D& NewValue)
+{
+	Parameter = NewValue;
+	LastChange.CurrentValue = NewValue;
+
+	// construct Scale as delta from initial position
+	FVector2D Delta = LastChange.GetChangeDelta() * ScaleMultiplier;
+
+	if (bUseEqualScaling)
+	{
+		Delta = FVector2D(Delta.X + Delta.Y);
+	}
+
+	FTransform NewTransform = InitialTransform;
+	FVector StartScale = InitialTransform.GetScale3D();
+	FVector NewScale = StartScale + Delta.X*CurScaleAxisX + Delta.Y*CurScaleAxisY;
+
+	if (bClampToZero)
+	{
+		NewScale = FVector::Max(NewScale, FVector::ZeroVector);
+	}
+
+	NewTransform.SetScale3D(NewScale);
+
+	// apply position constraint
+	//FVector SnappedPos;
+	//if (PositionConstraintFunction(NewTransform.GetScale(), SnappedPos))
+	//{
+	//	FVector PlanePos = GizmoMath::ProjectPointOntoPlane(SnappedPos, CurTranslationOrigin, CurTranslationNormal);
+	//	NewTransform.SetTranslation(PlanePos);
+	//}
+
+	TransformSource->SetTransform(NewTransform);
+
+	OnParameterChanged.Broadcast(this, LastChange);
+}
+
+void UGizmoPlaneScaleParameterSource::BeginModify()
+{
+	check(AxisSource);
+
+	LastChange = FGizmoVec2ParameterChange(Parameter);
+
+	// save initial transformation and axis information
+	InitialTransform = TransformSource->GetTransform();
+	CurScaleOrigin = AxisSource->GetOrigin();
+	AxisSource->GetAxisFrame(CurScaleNormal, CurScaleAxisX, CurScaleAxisY);
+}
+
+void UGizmoPlaneScaleParameterSource::EndModify()
+{
+}
\ No newline at end of file
diff --git a/Engine/Source/Runtime/InteractiveToolsFramework/Private/BaseGizmos/RepositionableTransformGizmo.cpp b/Engine/Source/Runtime/InteractiveToolsFramework/Private/BaseGizmos/RepositionableTransformGizmo.cpp
--- a/Engine/Source/Runtime/InteractiveToolsFramework/Private/BaseGizmos/RepositionableTransformGizmo.cpp
+++ b/Engine/Source/Runtime/InteractiveToolsFramework/Private/BaseGizmos/RepositionableTransformGizmo.cpp
@@ -77,19 +77,19 @@ void URepositionableTransformGizmo::SetActiveTarget(UTransformProxy* Target, ITo
 	// modify their alignment functions separately if we need to.
 	if (GizmoActor->TranslateX != nullptr)
 	{
-		UInteractiveGizmo* SubGizmo = AddAxisTranslationGizmo(GizmoActor->TranslateX, GizmoComponent, AxisXSource, TransformSource, RepositionStateTarget);
+		UInteractiveGizmo* SubGizmo = AddAxisTranslationGizmo(GizmoActor->TranslateX, GizmoComponent, AxisXSource, TransformSource, RepositionStateTarget, 0);
 		ModifyPivotAxisGizmo(SubGizmo);
 		PivotAlignmentGizmos.Add(SubGizmo);
 	}
 	if (GizmoActor->TranslateY != nullptr)
 	{
-		UInteractiveGizmo* SubGizmo = AddAxisTranslationGizmo(GizmoActor->TranslateY, GizmoComponent, AxisYSource, TransformSource, RepositionStateTarget);
+		UInteractiveGizmo* SubGizmo = AddAxisTranslationGizmo(GizmoActor->TranslateY, GizmoComponent, AxisYSource, TransformSource, RepositionStateTarget, 1);
 		ModifyPivotAxisGizmo(SubGizmo);
 		PivotAlignmentGizmos.Add(SubGizmo);
 	}
 	if (GizmoActor->TranslateZ != nullptr)
 	{
-		UInteractiveGizmo* SubGizmo = AddAxisTranslationGizmo(GizmoActor->TranslateZ, GizmoComponent, AxisZSource, TransformSource, RepositionStateTarget);
+		UInteractiveGizmo* SubGizmo = AddAxisTranslationGizmo(GizmoActor->TranslateZ, GizmoComponent, AxisZSource, TransformSource, RepositionStateTarget, 2);
 		ModifyPivotAxisGizmo(SubGizmo);
 		PivotAlignmentGizmos.Add(SubGizmo);
 	}
@@ -97,19 +97,19 @@ void URepositionableTransformGizmo::SetActiveTarget(UTransformProxy* Target, ITo
 
 	if (GizmoActor->TranslateYZ != nullptr)
 	{
-		UInteractiveGizmo* SubGizmo = AddPlaneTranslationGizmo(GizmoActor->TranslateYZ, GizmoComponent, AxisXSource, TransformSource, RepositionStateTarget);
+		UInteractiveGizmo* SubGizmo = AddPlaneTranslationGizmo(GizmoActor->TranslateYZ, GizmoComponent, AxisXSource, TransformSource, RepositionStateTarget, 1, 2);
 		ModifyPivotPlaneGizmo(SubGizmo);
 		PivotAlignmentGizmos.Add(SubGizmo);
 	}
 	if (GizmoActor->TranslateXZ != nullptr)
 	{
-		UInteractiveGizmo* SubGizmo = AddPlaneTranslationGizmo(GizmoActor->TranslateXZ, GizmoComponent, AxisYSource, TransformSource, RepositionStateTarget);
+		UInteractiveGizmo* SubGizmo = AddPlaneTranslationGizmo(GizmoActor->TranslateXZ, GizmoComponent, AxisYSource, TransformSource, RepositionStateTarget, 2, 0);
 		ModifyPivotPlaneGizmo(SubGizmo);
 		PivotAlignmentGizmos.Add(SubGizmo);
 	}
 	if (GizmoActor->TranslateXY != nullptr)
 	{
-		UInteractiveGizmo* SubGizmo = AddPlaneTranslationGizmo(GizmoActor->TranslateXY, GizmoComponent, AxisZSource, TransformSource, RepositionStateTarget);
+		UInteractiveGizmo* SubGizmo = AddPlaneTranslationGizmo(GizmoActor->TranslateXY, GizmoComponent, AxisZSource, TransformSource, RepositionStateTarget, 0, 1);
 		ModifyPivotPlaneGizmo(SubGizmo);
 		PivotAlignmentGizmos.Add(SubGizmo);
 	}
diff --git a/Engine/Source/Runtime/InteractiveToolsFramework/Public/BaseGizmos/CombinedTransformGizmo.h b/Engine/Source/Runtime/InteractiveToolsFramework/Public/BaseGizmos/CombinedTransformGizmo.h
--- a/Engine/Source/Runtime/InteractiveToolsFramework/Public/BaseGizmos/CombinedTransformGizmo.h
+++ b/Engine/Source/Runtime/InteractiveToolsFramework/Public/BaseGizmos/CombinedTransformGizmo.h
@@ -246,6 +246,33 @@ public:
 };
 
 
+/**
+ * FToolContextOptionalToggle is used to store a boolean flag where the value of 
+ * the boolean may either be set directly, or it may be set by querying some external context.
+ * This struct does not directly do anything, it just wraps up the multiple flags/states
+ * needed to provide such functionality
+ */
+struct INTERACTIVETOOLSFRAMEWORK_API FToolContextOptionalToggle
+{
+	bool bEnabledDirectly = false;
+	bool bEnabledInContext = false;
+	bool bInheritFromContext = false;
+
+	FToolContextOptionalToggle() {}
+	FToolContextOptionalToggle(bool bEnabled, bool bSetInheritFromContext)
+	{
+		bEnabledDirectly = bEnabled;
+		bInheritFromContext = bSetInheritFromContext;
+	}
+
+	void UpdateContextValue(bool bNewValue) { bEnabledInContext = bNewValue; }
+	bool InheritFromContext() const { return bInheritFromContext; }
+
+	/** @return true if Toggle is currently set to Enabled/On, under the current configuration */
+	bool IsEnabled() const { return (bInheritFromContext) ? bEnabledInContext : bEnabledDirectly; }
+};
+
+
 /**
  * UCombinedTransformGizmo provides standard Transformation Gizmo interactions,
  * applied to a UTransformProxy target object. By default the Gizmo will be
@@ -310,7 +337,6 @@ public:
 	virtual void Shutdown() override;
 	virtual void Tick(float DeltaTime) override;
 
-
 	/**
 	 * Set the active target object for the Gizmo
 	 * @param Target active target
@@ -372,11 +398,17 @@ public:
 	}
 
 	/**
-	 * If true, then when using world frame, Axis and Plane translation snap to the world grid via the ContextQueriesAPI (in PositionSnapFunction)
+	 * bSnapToWorldGrid controls whether any position snapping is applied, if possible, for Axis and Plane translations, via the ContextQueriesAPI
+	 * Despite the name, this flag controls both world-space grid snapping and relative snapping
 	 */
 	UPROPERTY()
 	bool bSnapToWorldGrid = true;
 
+	/**
+	 * Specify whether Relative snapping for Translations should be used in World frame mode. Relative snapping is always used in Local mode.
+	 */
+	FToolContextOptionalToggle RelativeTranslationSnapping = FToolContextOptionalToggle(true, true);
+
 	/**
 	 * Optional grid size which overrides the Context Grid
 	 */
@@ -399,6 +431,8 @@ public:
 	UPROPERTY()
 	bool bSnapToWorldRotGrid = true;
 
+
+
 	/**
 	 * Whether to use the World/Local coordinate system provided by the context via the ContextyQueriesAPI.
 	 */
@@ -413,6 +447,8 @@ public:
 	EToolContextCoordinateSystem CurrentCoordinateSystem = EToolContextCoordinateSystem::Local;
 
 
+
+
 	/**
 	 * Whether to use the Gizmo Mode provided by the context via the ContextyQueriesAPI.
 	 */
@@ -549,14 +585,16 @@ protected:
 		UPrimitiveComponent* AxisComponent, USceneComponent* RootComponent,
 		IGizmoAxisSource* AxisSource,
 		IGizmoTransformSource* TransformSource, 
-		IGizmoStateTarget* StateTarget);
+		IGizmoStateTarget* StateTarget,
+		int AxisIndex);
 
 	/** @return a new instance of the standard plane-translation Gizmo */
 	virtual UInteractiveGizmo* AddPlaneTranslationGizmo(
 		UPrimitiveComponent* AxisComponent, USceneComponent* RootComponent,
 		IGizmoAxisSource* AxisSource,
 		IGizmoTransformSource* TransformSource,
-		IGizmoStateTarget* StateTarget);
+		IGizmoStateTarget* StateTarget,
+		int XAxisIndex, int YAxisIndex);
 
 	/** @return a new instance of the standard axis-rotation Gizmo */
 	virtual UInteractiveGizmo* AddAxisRotationGizmo(
@@ -586,8 +624,10 @@ protected:
 		IGizmoTransformSource* TransformSource,
 		IGizmoStateTarget* StateTarget);
 
-	// Axis and Plane TransformSources use this function to execute worldgrid snap queries
+	// Axis and Plane TransformSources use these function to execute snapping queries
 	bool PositionSnapFunction(const FVector& WorldPosition, FVector& SnappedPositionOut) const;
+	bool PositionAxisDeltaSnapFunction(double AxisDelta, double& SnappedDeltaOut, int AxisIndex) const;
 	FQuat RotationSnapFunction(const FQuat& DeltaRotation) const;
+	bool RotationAxisAngleSnapFunction(double AxisAngleDelta, double& SnappedAxisAngleDeltaOut, int AxisIndex) const;
 
 };
diff --git a/Engine/Source/Runtime/InteractiveToolsFramework/Public/BaseGizmos/ParameterToTransformAdapters.h b/Engine/Source/Runtime/InteractiveToolsFramework/Public/BaseGizmos/ParameterToTransformAdapters.h
--- a/Engine/Source/Runtime/InteractiveToolsFramework/Public/BaseGizmos/ParameterToTransformAdapters.h
+++ b/Engine/Source/Runtime/InteractiveToolsFramework/Public/BaseGizmos/ParameterToTransformAdapters.h
@@ -6,7 +6,6 @@
 #include "BaseGizmos/GizmoInterfaces.h"
 #include "BaseGizmos/ParameterSourcesFloat.h"
 #include "BaseGizmos/ParameterSourcesVec2.h"
-#include "BaseGizmos/GizmoMath.h"
 #include "ParameterToTransformAdapters.generated.h"
 
 
@@ -30,6 +29,14 @@ class INTERACTIVETOOLSFRAMEWORK_API UGizmoAxisTranslationParameterSource : publi
 	GENERATED_BODY()
 public:
 
+
+	/**
+	 * Optional axis delta constraint function. This can be used to snap/constrain the "parameter delta" value,
+	 * for example to snap it to fixed increments along the axis. 
+	 * @return true if constrained delta was found and should be used, false to ignore
+	 */
+	TUniqueFunction<bool(double AxisDelta, double& SnappedDelta)> AxisDeltaConstraintFunction = [](double, double&) { return false; };
+
 	/**
 	 * Optional position constraint function. Called during interaction with the new transform origin.
 	 * To snap the transform to a new position, return as second value, and return true from your lambda.
@@ -43,45 +50,10 @@ public:
 		return Parameter;
 	}
 
-	virtual void SetParameter(float NewValue) override
-	{
-		Parameter = NewValue;
-		LastChange.CurrentValue = NewValue;
-
-		// construct translation as delta from initial position
-		FVector Translation = LastChange.GetChangeDelta() * CurTranslationAxis;
-
-		// translate the initial transform
-		FTransform NewTransform = InitialTransform;
-		NewTransform.AddToTranslation(Translation);
+	virtual void SetParameter(float NewValue) override;
 
-		// apply position constraint
-		FVector SnappedPos;
-		if (PositionConstraintFunction(NewTransform.GetTranslation(), SnappedPos))
-		{
-			FVector SnappedLinePos = GizmoMath::ProjectPointOntoLine(SnappedPos, CurTranslationOrigin, CurTranslationAxis);
-			NewTransform.SetTranslation(SnappedLinePos);
-		}
-
-		TransformSource->SetTransform(NewTransform);
-
-		OnParameterChanged.Broadcast(this, LastChange);
-	}
-
-	virtual void BeginModify()
-	{
-		check(AxisSource);
-
-		LastChange = FGizmoFloatParameterChange(Parameter);
-
-		InitialTransform = TransformSource->GetTransform();
-		CurTranslationAxis = AxisSource->GetDirection();
-		CurTranslationOrigin = AxisSource->GetOrigin();
-	}
-
-	virtual void EndModify()
-	{
-	}
+	virtual void BeginModify() override;
+	virtual void EndModify() override;
 
 
 public:
@@ -149,13 +121,22 @@ class INTERACTIVETOOLSFRAMEWORK_API UGizmoPlaneTranslationParameterSource : publ
 	GENERATED_BODY()
 public:
 
+	/**
+	 * Optional axis delta constraint function. This can be used to snap/constrain the "parameter delta" value,
+	 * for example to snap it to fixed increments along the axis. 
+	 * The AxisDeltaConstraintFunction is used separately for the X and Y axes of the plane translation
+	 * @return true if constrained delta was found and should be used, false to ignore
+	 */
+	TUniqueFunction<bool(double AxisDelta, double& SnappedDelta)> AxisXDeltaConstraintFunction = [](double, double&) { return false; };
+	TUniqueFunction<bool(double AxisDelta, double& SnappedDelta)> AxisYDeltaConstraintFunction = [](double, double&) { return false; };
+
 	/**
 	 * Optional position constraint function. Called during interaction with the new transform origin.
 	 * To snap the transform to a new position, return as second value, and return true from your lambda.
 	 * Note that returned snap point will be projected onto the current translation origin/axis.
 	 * @return true if constraint point was found and should be used, false to ignore
 	 */
-	TUniqueFunction<bool(const FVector&, FVector&)> PositionConstraintFunction = [](const FVector&, FVector&) { return false; };
+	TUniqueFunction<bool(const FVector& Input, FVector& PreProjectedOutput)> PositionConstraintFunction = [](const FVector&, FVector&) { return false; };
 
 
 	virtual FVector2D GetParameter() const override
@@ -163,47 +144,10 @@ public:
 		return Parameter;
 	}
 
-	virtual void SetParameter(const FVector2D& NewValue) override
-	{
-		Parameter = NewValue;
-		LastChange.CurrentValue = NewValue;
-
-		// construct translation as delta from initial position
-		FVector2D Delta = LastChange.GetChangeDelta();
-		FVector Translation = Delta.X*CurTranslationAxisX + Delta.Y*CurTranslationAxisY;
-
-		// apply translation to initial transform
-		FTransform NewTransform = InitialTransform;
-		NewTransform.AddToTranslation(Translation);
-
-		// apply position constraint
-		FVector SnappedPos;
-		if (PositionConstraintFunction(NewTransform.GetTranslation(), SnappedPos))
-		{
-			FVector PlanePos = GizmoMath::ProjectPointOntoPlane(SnappedPos, CurTranslationOrigin, CurTranslationNormal);
-			NewTransform.SetTranslation(PlanePos);
-		}
-
-		TransformSource->SetTransform(NewTransform);
-
-		OnParameterChanged.Broadcast(this, LastChange);
-	}
-
-	virtual void BeginModify()
-	{
-		check(AxisSource);
+	virtual void SetParameter(const FVector2D& NewValue) override;
 
-		LastChange = FGizmoVec2ParameterChange(Parameter);
-
-		// save initial transformation and axis information
-		InitialTransform = TransformSource->GetTransform();
-		CurTranslationOrigin = AxisSource->GetOrigin();
-		AxisSource->GetAxisFrame(CurTranslationNormal, CurTranslationAxisX, CurTranslationAxisY);
-	}
-
-	virtual void EndModify()
-	{
-	}
+	virtual void BeginModify() override;
+	virtual void EndModify() override;
 
 
 public:
@@ -278,6 +222,14 @@ class INTERACTIVETOOLSFRAMEWORK_API UGizmoAxisRotationParameterSource : public U
 {
 	GENERATED_BODY()
 public:
+
+	/**
+	 * Optional axis-delta-angle constraint function. This can be used to snap/constrain the "angle delta" value,
+	 * for example to snap it to fixed angle steps around the axis. 
+	 * @return true if constrained angle was found and should be used, false to ignore
+	 */
+	TUniqueFunction<bool(double AngleDelta, double& SnappedAngleDelta)> AngleDeltaConstraintFunction = [](double, double&) { return false; };
+
 	/**
 	 * Optional rotation constraint function. Called during interaction with the rotation delta
 	 * To snap the rotation delta, return the snapped quat
@@ -290,48 +242,10 @@ public:
 		return Angle;
 	}
 
-	virtual void SetParameter(float NewValue) override
-	{
-		Angle = NewValue;
-		LastChange.CurrentValue = NewValue;
-
-		// construct rotation as delta from initial position
-		FQuat DeltaRotation(CurRotationAxis, LastChange.GetChangeDelta());
-		DeltaRotation = RotationConstraintFunction(DeltaRotation);
-
-		// rotate the vector from the rotation origin to the transform origin, 
-		// to get the translation of the origin produced by the rotation
-		FVector DeltaPosition = InitialTransform.GetLocation() - CurRotationOrigin;
-		DeltaPosition = DeltaRotation * DeltaPosition;
-		FVector NewLocation = CurRotationOrigin + DeltaPosition;
+	virtual void SetParameter(float NewValue) override;
 
-		// rotate the initial transform by the rotation
-		FQuat NewRotation = DeltaRotation * InitialTransform.GetRotation();
-
-		// construct new transform
-		FTransform NewTransform = InitialTransform;
-		NewTransform.SetLocation(NewLocation);
-		NewTransform.SetRotation(NewRotation);
-		TransformSource->SetTransform(NewTransform);
-
-		OnParameterChanged.Broadcast(this, LastChange);
-	}
-
-	virtual void BeginModify()
-	{
-		check(AxisSource != nullptr);
-
-		LastChange = FGizmoFloatParameterChange(Angle);
-
-		// save initial transformation and axis information
-		InitialTransform = TransformSource->GetTransform();
-		CurRotationAxis = AxisSource->GetDirection();
-		CurRotationOrigin = AxisSource->GetOrigin();
-	}
-
-	virtual void EndModify()
-	{
-	}
+	virtual void BeginModify() override;
+	virtual void EndModify() override;
 
 
 public:
@@ -422,51 +336,10 @@ public:
 		return Parameter;
 	}
 
-	virtual void SetParameter(const FVector2D& NewValue) override
-	{
-		Parameter = NewValue;
-		LastChange.CurrentValue = NewValue;
-
-		// Convert 2D parameter delta to a 1D uniform scale change
-		// This possibly be exposed as a TFunction to allow customization?
-		double SignedDelta = LastChange.GetChangeDelta().X + LastChange.GetChangeDelta().Y;
-		SignedDelta *= ScaleMultiplier;
-		SignedDelta += 1.0;
-
-		FTransform NewTransform = InitialTransform;
-		FVector StartScale = InitialTransform.GetScale3D();
-		FVector NewScale = SignedDelta * StartScale;
-		NewTransform.SetScale3D(NewScale);
-
-		// apply position constraint
-		//FVector SnappedPos;
-		//if (PositionConstraintFunction(NewTransform.GetScale(), SnappedPos))
-		//{
-		//	FVector PlanePos = GizmoMath::ProjectPointOntoPlane(SnappedPos, CurTranslationOrigin, CurTranslationNormal);
-		//	NewTransform.SetTranslation(PlanePos);
-		//}
-
-		TransformSource->SetTransform(NewTransform);
-
-		OnParameterChanged.Broadcast(this, LastChange);
-	}
+	virtual void SetParameter(const FVector2D& NewValue) override;
 
-	virtual void BeginModify()
-	{
-		check(AxisSource);
-
-		LastChange = FGizmoVec2ParameterChange(Parameter);
-
-		// save initial transformation and axis information
-		InitialTransform = TransformSource->GetTransform();
-		CurScaleOrigin = AxisSource->GetOrigin();
-		// note: currently not used!
-		AxisSource->GetAxisFrame(CurScaleNormal, CurScaleAxisX, CurScaleAxisY);
-	}
-
-	virtual void EndModify()
-	{
-	}
+	virtual void BeginModify() override;
+	virtual void EndModify() override;
 
 
 public:
@@ -558,56 +431,10 @@ public:
 		return Parameter;
 	}
 
-	virtual void SetParameter(float NewValue) override
-	{
-		Parameter = NewValue;
-		LastChange.CurrentValue = NewValue;
-
-		// construct translation as delta from initial position
-		float ScaleDelta = LastChange.GetChangeDelta() * ScaleMultiplier;
-
-		// translate the initial transform
-		FTransform NewTransform = InitialTransform;
-		FVector StartScale = InitialTransform.GetScale3D();
-		FVector NewScale = StartScale + ScaleDelta * CurScaleAxis;
-
-		if (bClampToZero)
-		{
-			NewScale = FVector::Max(FVector::ZeroVector, NewScale);
-		}
-
-		NewTransform.SetScale3D(NewScale);
-
-
-		// apply position constraint
-		//FVector SnappedPos;
-		//if (ScaleConstraintFunction(NewTransform.GetTranslation(), SnappedPos))
-		//{
-		//	FVector SnappedLinePos = GizmoMath::ProjectPointOntoLine(SnappedPos, CurTranslationOrigin, CurTranslationAxis);
-		//	NewTransform.SetTranslation(SnappedLinePos);
-		//}
-
-		TransformSource->SetTransform(NewTransform);
-
-		OnParameterChanged.Broadcast(this, LastChange);
-	}
-
-	virtual void BeginModify()
-	{
-		check(AxisSource);
-
-		LastChange = FGizmoFloatParameterChange(Parameter);
-
-		InitialTransform = TransformSource->GetTransform();
-		
-		CurScaleAxis = AxisSource->GetDirection();
-		CurScaleOrigin = AxisSource->GetOrigin();
-	}
-
-	virtual void EndModify()
-	{
-	}
+	virtual void SetParameter(float NewValue) override;
 
+	virtual void BeginModify() override;
+	virtual void EndModify() override;
 
 public:
 	/** The Parameter line-equation value is converted to a 3D Translation along this Axis */
@@ -698,58 +525,10 @@ public:
 		return Parameter;
 	}
 
-	virtual void SetParameter(const FVector2D& NewValue) override
-	{
-		Parameter = NewValue;
-		LastChange.CurrentValue = NewValue;
-
-		// construct Scale as delta from initial position
-		FVector2D Delta = LastChange.GetChangeDelta() * ScaleMultiplier;
-
-		if (bUseEqualScaling)
-		{
-			Delta = FVector2D(Delta.X + Delta.Y);
-		}
-
-		FTransform NewTransform = InitialTransform;
-		FVector StartScale = InitialTransform.GetScale3D();
-		FVector NewScale = StartScale + Delta.X*CurScaleAxisX + Delta.Y*CurScaleAxisY;
-
-		if (bClampToZero)
-		{
-			NewScale = FVector::Max(NewScale, FVector::ZeroVector);
-		}
-
-		NewTransform.SetScale3D(NewScale);
-
-		// apply position constraint
-		//FVector SnappedPos;
-		//if (PositionConstraintFunction(NewTransform.GetScale(), SnappedPos))
-		//{
-		//	FVector PlanePos = GizmoMath::ProjectPointOntoPlane(SnappedPos, CurTranslationOrigin, CurTranslationNormal);
-		//	NewTransform.SetTranslation(PlanePos);
-		//}
-
-		TransformSource->SetTransform(NewTransform);
-
-		OnParameterChanged.Broadcast(this, LastChange);
-	}
-
-	virtual void BeginModify()
-	{
-		check(AxisSource);
-
-		LastChange = FGizmoVec2ParameterChange(Parameter);
-
-		// save initial transformation and axis information
-		InitialTransform = TransformSource->GetTransform();
-		CurScaleOrigin = AxisSource->GetOrigin();
-		AxisSource->GetAxisFrame(CurScaleNormal, CurScaleAxisX, CurScaleAxisY);
-	}
+	virtual void SetParameter(const FVector2D& NewValue) override;
 
-	virtual void EndModify()
-	{
-	}
+	virtual void BeginModify() override;
+	virtual void EndModify() override;
 
 
 public:
diff --git a/Engine/Source/Runtime/InteractiveToolsFramework/Public/ToolContextInterfaces.h b/Engine/Source/Runtime/InteractiveToolsFramework/Public/ToolContextInterfaces.h
--- a/Engine/Source/Runtime/InteractiveToolsFramework/Public/ToolContextInterfaces.h
+++ b/Engine/Source/Runtime/InteractiveToolsFramework/Public/ToolContextInterfaces.h
@@ -157,6 +157,9 @@ struct INTERACTIVETOOLSFRAMEWORK_API FToolContextSnappingConfiguration
 	bool bEnableRotationGridSnapping = false;
 	/** Specify rotation snapping step size */
 	FRotator RotationGridAngles = FRotator::ZeroRotator;
+
+	/** Specify whether Position snapping in World Coordinate System should be Absolute or Relative */
+	bool bEnableAbsoluteWorldSnapping = false;
 };
 
 
