ToolsFramework: fix issue where toggling a modifier key could result in a mouse capture starting

diff --git a/Engine/Source/Runtime/InteractiveToolsFramework/Private/InputRouter.cpp b/Engine/Source/Runtime/InteractiveToolsFramework/Private/InputRouter.cpp
--- a/Engine/Source/Runtime/InteractiveToolsFramework/Private/InputRouter.cpp
+++ b/Engine/Source/Runtime/InteractiveToolsFramework/Private/InputRouter.cpp
@@ -76,7 +76,23 @@ bool UInputRouter::PostInputEvent(const FInputDeviceState& Input)
 				(LastMouseInputState.bCmdKeyDown != Input.bCmdKeyDown))
 			{
 				LastMouseInputState.SetModifierKeyStates(Input.bShiftKeyDown, Input.bAltKeyDown, Input.bCtrlKeyDown, Input.bCmdKeyDown);
-				PostInputEvent_Mouse(LastMouseInputState);
+
+				// cannot call PostInputEvent_Mouse() to propagate modifier key state update because if
+				// there is no active capture it may result in one starting!
+				//PostInputEvent_Mouse(LastMouseInputState);
+
+				if (ActiveLeftCapture != nullptr )
+				{
+					HandleCapturedMouseInput(Input);
+				}
+				else
+				{
+					bool bHoverStateUpdated = ProcessMouseHover(Input);
+					if (bHoverStateUpdated && bAutoInvalidateOnHover)
+					{
+						TransactionsAPI->PostInvalidation();
+					}
+				}
 			}
 		}
 
