InteractiveToolsFramework: Add a couple ensures to transform proxy begin/end change calls to avoid crashes in misuse.

diff --git a/Engine/Source/Runtime/InteractiveToolsFramework/Private/BaseGizmos/TransformProxy.cpp b/Engine/Source/Runtime/InteractiveToolsFramework/Private/BaseGizmos/TransformProxy.cpp
--- a/Engine/Source/Runtime/InteractiveToolsFramework/Private/BaseGizmos/TransformProxy.cpp
+++ b/Engine/Source/Runtime/InteractiveToolsFramework/Private/BaseGizmos/TransformProxy.cpp
@@ -220,6 +220,8 @@ void FTransformProxyChangeSource::BeginChange()
 {
 	if (Proxy.IsValid())
 	{
+		ensureMsgf(!ActiveChange, TEXT("FTransformProxyChangeSource: BeginChange called without finishing previous change via EndChange."));
+
 		ActiveChange = MakeUnique<FTransformProxyChange>();
 		ActiveChange->From = Proxy->GetTransform();
 		ActiveChange->bSetPivotMode = bOverrideSetPivotMode ? true : Proxy->bSetPivotMode;
@@ -239,6 +241,13 @@ TUniquePtr<FToolCommandChange> FTransformProxyChangeSource::EndChange()
 {
 	if (Proxy.IsValid())
 	{
+		if (!ensureMsgf(ActiveChange, TEXT("FTransformProxyChangeSource: EndChange called without a pending change.")))
+		{
+			// One way this error could occur is if there is (incorrectly) one Begin/EndChange pair nested inside another one,
+			// so the previous EndChange call closes out the pending transaction.
+			return TUniquePtr<FToolCommandChange>();
+		}
+
 		if (ActiveChange->bSetPivotMode)
 		{
 			Proxy->EndPivotEditSequence();
