InteractiveToolsFramework: Adjust how the CombinedTransformGizmo handles the scale axis handles in cases where only two axes are requested. The old approach had the handles react to the camera direction at all times, causing issues in the UV Editor and inconsistent orientations of the handles relative to the plane in the case of only two axes. The new approach only is reactive to the camera in cases of one or three axes and locks the handles into a single plane if exactly two axes are configured in the gizmo.

diff --git a/Engine/Source/Runtime/InteractiveToolsFramework/Private/BaseGizmos/CombinedTransformGizmo.cpp b/Engine/Source/Runtime/InteractiveToolsFramework/Private/BaseGizmos/CombinedTransformGizmo.cpp
--- a/Engine/Source/Runtime/InteractiveToolsFramework/Private/BaseGizmos/CombinedTransformGizmo.cpp
+++ b/Engine/Source/Runtime/InteractiveToolsFramework/Private/BaseGizmos/CombinedTransformGizmo.cpp
@@ -199,28 +199,70 @@ ACombinedTransformGizmoActor* ACombinedTransformGizmoActor::ConstructCustom3Axis
 
 
 
-	auto MakeAxisScaleFunc = [&](const FLinearColor& Color, const FVector& Axis0, const FVector& Axis1)
+	auto MakeAxisScaleFunc = [&](const FLinearColor& Color, const FVector& Axis0, const FVector& Axis1, bool bLockSinglePlane)
 	{
 		UGizmoRectangleComponent* ScaleComponent = AddDefaultRectangleComponent(World, NewActor, GizmoViewContext, Color, Axis0, Axis1);
 		ScaleComponent->OffsetX = 140.0f; ScaleComponent->OffsetY = -10.0f;
 		ScaleComponent->LengthX = 7.0f; ScaleComponent->LengthY = 20.0f;
 		ScaleComponent->Thickness = GizmoLineThickness;
-		ScaleComponent->bOrientYAccordingToCamera = true;
+		ScaleComponent->bOrientYAccordingToCamera = !bLockSinglePlane;
 		ScaleComponent->NotifyExternalPropertyUpdates();
 		ScaleComponent->SegmentFlags = 0x1 | 0x2 | 0x4; // | 0x8;
 		return ScaleComponent;
 	};
+
+	// This is designed so we can properly handle the visual orientations of the scale handles under the condition of a
+	// planar gizmo (such as in the UV Editor).
+	// In this case we want to lock the handle on to the other axis of the plane, rather than use the component's camera orientation option. This requires
+	// both tracking how many axes are being requested and also *which* axes are requested, in order to configure the correct planar basis vectors.
+	// In the case of a single axis, we have to pick a cross axis arbitrarily, but we also keep the auto orientation mode on the component active, so the initial
+	// choice isn't as critical. If we want to some day have a single axis handle that is locked, we may need to revisit this again.
+	auto ConfigureAdditionalAxis = [&Elements](ETransformGizmoSubElements AxisToTest, int32& TotalAxisCount, FVector& NewPerpendicularAxis) {
+		if ((Elements & ETransformGizmoSubElements::ScaleAxisX & AxisToTest) != ETransformGizmoSubElements::None)
+		{
+			TotalAxisCount++;
+			NewPerpendicularAxis = FVector(1, 0, 0);
+			return;
+		}
+		if ((Elements & ETransformGizmoSubElements::ScaleAxisY & AxisToTest) != ETransformGizmoSubElements::None)
+		{
+			TotalAxisCount++;
+			NewPerpendicularAxis = FVector(0, 1, 0);
+			return;
+		}
+		if ((Elements & ETransformGizmoSubElements::ScaleAxisZ & AxisToTest) != ETransformGizmoSubElements::None)
+		{
+			TotalAxisCount++;
+			NewPerpendicularAxis = FVector(0, 0, 1);
+			return;
+		}
+	};
+
 	if ((Elements & ETransformGizmoSubElements::ScaleAxisX) != ETransformGizmoSubElements::None)
 	{
-		NewActor->AxisScaleX = MakeAxisScaleFunc(FLinearColor::Red, FVector(1, 0, 0), FVector(0, 0, 1));
+		int32 TotalAxisCount = 1;
+		FVector PerpendicularAxis(0,1,0);
+		ConfigureAdditionalAxis(ETransformGizmoSubElements::ScaleAxisY, TotalAxisCount, PerpendicularAxis);
+		ConfigureAdditionalAxis(ETransformGizmoSubElements::ScaleAxisZ, TotalAxisCount, PerpendicularAxis);
+		NewActor->AxisScaleX = MakeAxisScaleFunc(FLinearColor::Red, FVector(1, 0, 0), PerpendicularAxis, TotalAxisCount == 2);
 	}
+
 	if ((Elements & ETransformGizmoSubElements::ScaleAxisY) != ETransformGizmoSubElements::None)
 	{
-		NewActor->AxisScaleY = MakeAxisScaleFunc(FLinearColor::Green, FVector(0, 1, 0), FVector(0, 0, 1));
+		int32 TotalAxisCount = 1;
+		FVector PerpendicularAxis(1, 0, 0);
+		ConfigureAdditionalAxis(ETransformGizmoSubElements::ScaleAxisX, TotalAxisCount, PerpendicularAxis);
+		ConfigureAdditionalAxis(ETransformGizmoSubElements::ScaleAxisZ, TotalAxisCount, PerpendicularAxis);
+		NewActor->AxisScaleY = MakeAxisScaleFunc(FLinearColor::Green, FVector(0, 1, 0), PerpendicularAxis, TotalAxisCount == 2);
 	}
+
 	if ((Elements & ETransformGizmoSubElements::ScaleAxisZ) != ETransformGizmoSubElements::None)
 	{
-		NewActor->AxisScaleZ = MakeAxisScaleFunc(FLinearColor::Blue, FVector(0, 0, 1), FVector(1, 0, 0));
+		int32 TotalAxisCount = 1;
+		FVector PerpendicularAxis(1, 0, 0);
+		ConfigureAdditionalAxis(ETransformGizmoSubElements::ScaleAxisY, TotalAxisCount, PerpendicularAxis);
+		ConfigureAdditionalAxis(ETransformGizmoSubElements::ScaleAxisX, TotalAxisCount, PerpendicularAxis);
+		NewActor->AxisScaleZ = MakeAxisScaleFunc(FLinearColor::Blue, FVector(0, 0, 1), PerpendicularAxis, TotalAxisCount == 2);
 	}
 
 
