InteractiveToolsFramework: Fix click drag behaviors passing bogus world ray to their targets when a modifier key is pressed during a drag.

diff --git a/Engine/Source/Runtime/InteractiveToolsFramework/Private/BaseBehaviors/ClickDragBehavior.cpp b/Engine/Source/Runtime/InteractiveToolsFramework/Private/BaseBehaviors/ClickDragBehavior.cpp
--- a/Engine/Source/Runtime/InteractiveToolsFramework/Private/BaseBehaviors/ClickDragBehavior.cpp
+++ b/Engine/Source/Runtime/InteractiveToolsFramework/Private/BaseBehaviors/ClickDragBehavior.cpp
@@ -50,6 +50,13 @@ FInputCaptureUpdate UClickDragInputBehavior::UpdateCapture(const FInputDeviceSta
 		Modifiers.UpdateModifiers(Input, Target);
 	}
 
+	// We have to check the device before going further because we get passed captures from 
+	// keyboard for modifier key press/releases, and those don't have valid ray data.
+	if ((GetSupportedDevices() & Input.InputDevice) == EInputDevices::None)
+	{
+		return FInputCaptureUpdate::Continue();
+	}
+
 	if (IsReleased(Input)) 
 	{
 		OnClickReleaseInternal(Input, Data);
