ToolsFramework: add ray-distance test for rectangle border segments in UGizmoRectangleComponent::LineTraceComponent(), to make hit-testing of combined-gizmo plane-translate and scaling handles more consistent with axis-translate/rotate handles. Previously the cursor could be over the scaling handle line segments w/o a hit being registered.

diff --git a/Engine/Source/Runtime/InteractiveToolsFramework/Private/BaseGizmos/GizmoRectangleComponent.cpp b/Engine/Source/Runtime/InteractiveToolsFramework/Private/BaseGizmos/GizmoRectangleComponent.cpp
--- a/Engine/Source/Runtime/InteractiveToolsFramework/Private/BaseGizmos/GizmoRectangleComponent.cpp
+++ b/Engine/Source/Runtime/InteractiveToolsFramework/Private/BaseGizmos/GizmoRectangleComponent.cpp
@@ -18,7 +18,8 @@ namespace GizmoRectangleComponentLocals
 		float OffsetX, float OffsetY, float LengthX, float LengthY,
 		bool bUseWorldAxes, bool bOrientYAccordingToCamera, 
 		TFunctionRef<FVector(const FVector&)> VectorTransform,
-		TArray<FVector>& CornersOut)
+		TArray<FVector>& CornersOut,
+		float& PixelToWorldScaleOut)
 	{
 		FVector UseDirectionX = (bUseWorldAxes) ? DirectionX : VectorTransform(DirectionX);
 		FVector UseDirectionY = (bUseWorldAxes) ? DirectionY : VectorTransform(DirectionY);
@@ -72,6 +73,8 @@ namespace GizmoRectangleComponentLocals
 		CornersOut[2] = WorldOrigin + UseOffsetLengthX * UseDirectionX + UseOffsetLengthY * UseDirectionY;
 		CornersOut[3] = WorldOrigin + UseOffsetX * UseDirectionX + UseOffsetLengthY * UseDirectionY;
 
+		PixelToWorldScaleOut = LengthScale;
+
 		return true;
 	}
 }
@@ -119,10 +122,11 @@ public:
 				
 				TArray<FVector> Corners;
 				bool bIsViewDependent = (bExternalIsViewDependent) ? (*bExternalIsViewDependent) : false;
+				float PixelToWorldScale = 1.0;
 				bool bRenderVisibility = GetWorldCorners(bIsViewDependent, View, Origin, DirectionX, DirectionY,
 					OffsetX, OffsetY, LengthX, LengthY, bWorldAxis, bOrientYAccordingToCamera,
 					[&LocalToWorldMatrix](const FVector& VectorIn) { return FVector{ LocalToWorldMatrix.TransformVector(VectorIn) }; },
-					Corners);
+					Corners, PixelToWorldScale);
 
 				if (!bRenderVisibility)
 				{
@@ -153,24 +157,6 @@ public:
 					PDI->DrawLine(Corners[3], Corners[0], Color, SDPG_Foreground, UseThickness, 0.0f, true);
 				}
 
-
-				/* 
-				 
-				// Draw solid square. Will be occluded but can be used w/ custom depth to show hidden
-				 
-				FDynamicMeshBuilder MeshBuilder(View->GetFeatureLevel());
-				MeshBuilder.AddVertex(Point00, FVector2D::ZeroVector, FVector::ZeroVector, FVector::ZeroVector, FVector::ZeroVector, UseColor);
-				MeshBuilder.AddVertex(Point10, FVector2D::ZeroVector, FVector::ZeroVector, FVector::ZeroVector, FVector::ZeroVector, UseColor);
-				MeshBuilder.AddVertex(Point11, FVector2D::ZeroVector, FVector::ZeroVector, FVector::ZeroVector, FVector::ZeroVector, UseColor);
-				MeshBuilder.AddVertex(Point01, FVector2D::ZeroVector, FVector::ZeroVector, FVector::ZeroVector, FVector::ZeroVector, UseColor);
-				MeshBuilder.AddTriangle(0, 1, 2);
-				MeshBuilder.AddTriangle(0, 2, 3);
-
-				//FMaterialRenderProxy* MaterialRenderProxy = Material->GetRenderProxy();
-				FMaterialRenderProxy* MaterialRenderProxy = GEngine->VertexColorMaterial->GetRenderProxy();
-				MeshBuilder.GetMesh(FMatrix::Identity, MaterialRenderProxy, SDPG_Foreground, true, false, ViewIndex, Collector);
-				//MeshBuilder.GetMesh(FMatrix::Identity, MaterialRenderProxy, SDPG_World, true, false, ViewIndex, Collector);
-				*/
 			}
 		}
 	}
@@ -253,10 +239,11 @@ bool UGizmoRectangleComponent::LineTraceComponent(FHitResult& OutHit, const FVec
 	FVector UseOrigin = Transform.TransformPosition(FVector::ZeroVector);
 
 	TArray<FVector> Corners;
+	float PixelToWorldScale = 1.0;
 	bool bRenderVisibility = GetWorldCorners(bIsViewDependent, ToRawPtr(GizmoViewContext), UseOrigin, DirectionX, DirectionY, 
 		OffsetX, OffsetY, LengthX, LengthY,	bWorld, bOrientYAccordingToCamera,
 		[&Transform](const FVector& VectorIn) { return Transform.TransformVector(VectorIn); },
-		Corners);
+		Corners, PixelToWorldScale);
 
 	if (bRenderVisibility == false)
 	{
@@ -265,6 +252,38 @@ bool UGizmoRectangleComponent::LineTraceComponent(FHitResult& OutHit, const FVec
 
 	static const int Triangles[2][3] = { {0,1,2}, {0,2,3} };
 
+	// If ray is within pixel distance tolerance of boundary edges of rectangle,
+	// consider this a hit. Currently rectangle is rendered as a thick outline, 
+	// so if this is not done, and only the ray/try tests below are done, then
+	// the cursor can be directly over the line without a hit being registered.
+	double MinHitDistance = PixelHitDistanceThreshold * PixelToWorldScale;
+	bool bFoundLineHit = false;
+	for (int32 k = 0; k < 4; ++k)
+	{
+		FVector NearestArrow, NearestLine;
+		FMath::SegmentDistToSegmentSafe(Corners[k], Corners[(k+1)%4], Start, End, NearestArrow, NearestLine);
+		double Distance = FVector::Distance(NearestArrow, NearestLine);
+		if (Distance < MinHitDistance)
+		{
+			bFoundLineHit = true;
+			OutHit.Component = this;
+			OutHit.Distance = static_cast<float>(FVector::Distance(Start, NearestLine));
+			OutHit.ImpactPoint = NearestLine;
+			MinHitDistance = Distance;
+		}
+	}
+	if (bFoundLineHit)
+	{
+		FVector Edge1(Corners[Triangles[0][1]] - Corners[Triangles[0][0]]);
+		Edge1.Normalize();
+		FVector Edge2(Corners[Triangles[0][2]] - Corners[Triangles[0][0]]);
+		Edge2.Normalize();
+		OutHit.ImpactNormal = Edge2.Cross(Edge1);	// same code here as SegmentTriangleIntersection()
+		OutHit.ImpactNormal.Normalize();
+		return true;
+	}
+
+	// if line was not hit, try hitting rectangle triangles
 	for (int j = 0; j < 2; ++j)
 	{
 		const int* Triangle = Triangles[j];
