Disabled deprecation warnings in FToolBuilderState's special member functions.

diff --git a/Engine/Source/Runtime/InteractiveToolsFramework/Public/ToolContextInterfaces.h b/Engine/Source/Runtime/InteractiveToolsFramework/Public/ToolContextInterfaces.h
--- a/Engine/Source/Runtime/InteractiveToolsFramework/Public/ToolContextInterfaces.h
+++ b/Engine/Source/Runtime/InteractiveToolsFramework/Public/ToolContextInterfaces.h
@@ -62,6 +62,14 @@ struct INTERACTIVETOOLSFRAMEWORK_API FToolBuilderState
 
 	UE_DEPRECATED(5.1, "This has moved to a context object. See IAssetEditorContextInterface")
 	TWeakObjectPtr<UTypedElementSelectionSet> TypedElementSelectionSet;
+
+	PRAGMA_DISABLE_DEPRECATION_WARNINGS
+	FToolBuilderState() = default;
+	FToolBuilderState(FToolBuilderState&&) = default;
+	FToolBuilderState(const FToolBuilderState&) = default;
+	FToolBuilderState& operator=(FToolBuilderState&&) = default;
+	FToolBuilderState& operator=(const FToolBuilderState&) = default;
+	PRAGMA_ENABLE_DEPRECATION_WARNINGS
 };
 
 
