InteractiveToolsFramework: add UInteractiveToolManager::bInToolShutdown flag, set during DeactivateToolInternal and use in DeactivateTool() to prevent tool shutdown from being called while tool is actively shutting down. This can occur if higher-level Mode calls shutdown on (eg) editor delegates, which might be fired by a shutdown that is actively ocurring.

diff --git a/Engine/Source/Runtime/InteractiveToolsFramework/Private/InteractiveToolManager.cpp b/Engine/Source/Runtime/InteractiveToolsFramework/Private/InteractiveToolManager.cpp
--- a/Engine/Source/Runtime/InteractiveToolsFramework/Private/InteractiveToolManager.cpp
+++ b/Engine/Source/Runtime/InteractiveToolsFramework/Private/InteractiveToolManager.cpp
@@ -239,6 +239,10 @@ void UInteractiveToolManager::DeactivateTool(EToolSide Side, EToolShutdownType S
 	{
 		return;
 	}
+	if (bInToolShutdown)
+	{
+		return;
+	}
 
 	if (ActiveLeftTool != nullptr)
 	{
@@ -264,6 +268,8 @@ void UInteractiveToolManager::DeactivateToolInternal(EToolSide Side, EToolShutdo
 			return;
 		}
 
+		bInToolShutdown = true;
+
 		InputRouter->ForceTerminateSource(ActiveLeftTool);
 
 		ActiveLeftTool->Shutdown(ShutdownType);
@@ -277,6 +283,8 @@ void UInteractiveToolManager::DeactivateToolInternal(EToolSide Side, EToolShutdo
 		PostInvalidation();
 
 		OnToolEnded.Broadcast(this, DoneTool);
+
+		bInToolShutdown = false;
 	}
 }
 
diff --git a/Engine/Source/Runtime/InteractiveToolsFramework/Public/InteractiveToolManager.h b/Engine/Source/Runtime/InteractiveToolsFramework/Public/InteractiveToolManager.h
--- a/Engine/Source/Runtime/InteractiveToolsFramework/Public/InteractiveToolManager.h
+++ b/Engine/Source/Runtime/InteractiveToolsFramework/Public/InteractiveToolManager.h
@@ -350,6 +350,7 @@ protected:
 
 	FString ActiveLeftToolName;
 	FString ActiveRightToolName;
+	bool bInToolShutdown = false;
 
 	/** 
 	 * Tracks whether the last activated tool has made a request to store a tool selection.
