Added ability to disable base brush tool. Added toggle to skin paint tool while in selection mode.

diff --git a/Engine/Plugins/Experimental/MeshModelingToolsetExp/Source/MeshModelingToolsExp/Private/SkinWeightsPaintTool.cpp b/Engine/Plugins/Experimental/MeshModelingToolsetExp/Source/MeshModelingToolsExp/Private/SkinWeightsPaintTool.cpp
--- a/Engine/Plugins/Experimental/MeshModelingToolsetExp/Source/MeshModelingToolsExp/Private/SkinWeightsPaintTool.cpp
+++ b/Engine/Plugins/Experimental/MeshModelingToolsetExp/Source/MeshModelingToolsExp/Private/SkinWeightsPaintTool.cpp
@@ -18,6 +18,7 @@
 #include "Parameterization/MeshLocalParam.h"
 #include "Spatial/FastWinding.h"
 #include "Async/Async.h"
+#include "BaseGizmos/BrushStampIndicator.h"
 #include "DynamicMesh/MeshAdapterUtil.h"
 #include "Selection/PolygonSelectionMechanic.h"
 #include "Spatial/PointSetHashTable.h"
@@ -656,11 +657,6 @@ void USkinWeightsPaintTool::Setup()
 		EToolMessageLevel::UserNotification);
 }
 
-void USkinWeightsPaintTool::RegisterActions(FInteractiveToolActionSet& ActionSet)
-{
-	UDynamicMeshBrushTool::RegisterActions(ActionSet);
-}
-
 void USkinWeightsPaintTool::DrawHUD(FCanvas* Canvas, IToolsContextRenderAPI* RenderAPI)
 {
 	Super::DrawHUD(Canvas, RenderAPI);
@@ -677,11 +673,11 @@ void USkinWeightsPaintTool::Render(IToolsContextRenderAPI* RenderAPI)
 	{
 		Super::Render(RenderAPI);	
 	}
-	
-	if (PolygonSelectionMechanic && WeightToolProperties->EditingMode == EWeightEditMode::Selection)
+	else if (PolygonSelectionMechanic && WeightToolProperties->EditingMode == EWeightEditMode::Selection)
 	{
 		PolygonSelectionMechanic->Render(RenderAPI);
 	}
+	
 }
 
 void USkinWeightsPaintTool::OnTick(float DeltaTime)
@@ -1784,7 +1780,9 @@ void USkinWeightsPaintTool::OnSelectionModified()
 
 void USkinWeightsPaintTool::ToggleEditingMode()
 {
+	SetBrushEnabled(WeightToolProperties->EditingMode == EWeightEditMode::Brush);
 	PolygonSelectionMechanic->SetIsEnabled(WeightToolProperties->EditingMode == EWeightEditMode::Selection);
+	
 }
 
 void USkinWeightsPaintTool::OnPropertyModified(UObject* ModifiedObject, FProperty* ModifiedProperty)
diff --git a/Engine/Plugins/Experimental/MeshModelingToolsetExp/Source/MeshModelingToolsExp/Public/SkinWeightsPaintTool.h b/Engine/Plugins/Experimental/MeshModelingToolsetExp/Source/MeshModelingToolsExp/Public/SkinWeightsPaintTool.h
--- a/Engine/Plugins/Experimental/MeshModelingToolsetExp/Source/MeshModelingToolsExp/Public/SkinWeightsPaintTool.h
+++ b/Engine/Plugins/Experimental/MeshModelingToolsetExp/Source/MeshModelingToolsExp/Public/SkinWeightsPaintTool.h
@@ -297,7 +297,6 @@ public:
 	
 	// UInteractiveTool
 	virtual void Setup() override;
-	virtual void RegisterActions(FInteractiveToolActionSet& ActionSet) override;
 	virtual void DrawHUD(FCanvas* Canvas, IToolsContextRenderAPI* RenderAPI) override;
 	virtual void Render(IToolsContextRenderAPI* RenderAPI) override;
 
diff --git a/Engine/Source/Runtime/InteractiveToolsFramework/Private/BaseTools/BaseBrushTool.cpp b/Engine/Source/Runtime/InteractiveToolsFramework/Private/BaseTools/BaseBrushTool.cpp
--- a/Engine/Source/Runtime/InteractiveToolsFramework/Private/BaseTools/BaseBrushTool.cpp
+++ b/Engine/Source/Runtime/InteractiveToolsFramework/Private/BaseTools/BaseBrushTool.cpp
@@ -5,6 +5,7 @@
 #include "InteractiveGizmoManager.h"
 #include "ToolBuilderUtil.h"
 #include "InteractiveTool.h"
+#include "BaseBehaviors/ClickDragBehavior.h"
 #include "BaseGizmos/BrushStampIndicator.h"
 
 #include UE_INLINE_GENERATED_CPP_BY_NAME(BaseBrushTool)
@@ -28,6 +29,7 @@ UBaseBrushTool::UBaseBrushTool()
 	PropertyClass = UBrushBaseProperties::StaticClass();
 }
 
+static const int UBaseBrushTool_BrushSizeModifier = 0;
 
 void UBaseBrushTool::Setup()
 {
@@ -42,6 +44,15 @@ void UBaseBrushTool::Setup()
 	AddToolPropertySource(BrushProperties);
 
 	SetupBrushStampIndicator();
+
+	// add input behavior to click-drag while holding hotkey to adjust brush size
+	UClickDragInputBehavior* DragBehavior = NewObject<UClickDragInputBehavior>();
+	DragBehavior->Modifiers.RegisterModifier(UBaseBrushTool_BrushSizeModifier, [](const FInputDeviceState& Input)
+	{
+		return (Input.Keyboard.ActiveKey.Button == EKeys::B) && Input.Keyboard.ActiveKey.bDown;
+	});
+	DragBehavior->Initialize(this);
+	AddInputBehavior(DragBehavior);
 }
 
 
@@ -59,7 +70,40 @@ void UBaseBrushTool::OnPropertyModified(UObject* PropertySet, FProperty* Propert
 	}
 }
 
+FInputRayHit UBaseBrushTool::CanBeginClickDragSequence(const FInputDeviceRay& PressPos)
+{
+	if (bEnabled)
+	{
+		return Super::CanBeginClickDragSequence(PressPos);	
+	}
+
+	return FInputRayHit();
+}
+
+void UBaseBrushTool::OnClickPress(const FInputDeviceRay& PressPos)
+{
+	Super::OnClickPress(PressPos);
+}
+
+void UBaseBrushTool::OnClickDrag(const FInputDeviceRay& DragPos)
+{
+	Super::OnClickDrag(DragPos);
+}
+
+void UBaseBrushTool::OnClickRelease(const FInputDeviceRay& ReleasePos)
+{
+	Super::OnClickRelease(ReleasePos);
+}
+
+void UBaseBrushTool::OnTerminateDragSequence()
+{
+	Super::OnTerminateDragSequence();
+}
 
+void UBaseBrushTool::OnUpdateModifierState(int ModifierID, bool bIsOn)
+{
+	Super::OnUpdateModifierState(ModifierID, bIsOn);
+}
 
 void UBaseBrushTool::IncreaseBrushSizeAction()
 {
@@ -132,6 +176,12 @@ void UBaseBrushTool::DecreaseBrushFalloffAction()
 	NotifyOfPropertyChangeByTool(BrushProperties);
 }
 
+void UBaseBrushTool::SetBrushEnabled(bool bIsEnabled)
+{
+	bEnabled = bIsEnabled;
+	BrushStampIndicator->bVisible = bIsEnabled;
+}
+
 void UBaseBrushTool::RegisterActions(FInteractiveToolActionSet& ActionSet)
 {
 	ActionSet.RegisterAction(this, (int32)EStandardToolActions::BaseClientDefinedActionID + 10,
@@ -229,7 +279,6 @@ void UBaseBrushTool::OnEndDrag(const FRay& Ray)
 	bInBrushStroke = false;
 }
 
-
 bool UBaseBrushTool::OnUpdateHover(const FInputDeviceRay& DevicePos)
 {
 	FHitResult OutHit;
@@ -244,13 +293,13 @@ bool UBaseBrushTool::OnUpdateHover(const FInputDeviceRay& DevicePos)
 	return true;
 }
 
-
-
 void UBaseBrushTool::Render(IToolsContextRenderAPI* RenderAPI)
 {
-	UMeshSurfacePointTool::Render(RenderAPI);
-
-	UpdateBrushStampIndicator();
+	if (bEnabled)
+	{
+		UMeshSurfacePointTool::Render(RenderAPI);
+		UpdateBrushStampIndicator();
+	}
 }
 
 const FString BaseBrushIndicatorGizmoType = TEXT("BrushIndicatorGizmoType");
diff --git a/Engine/Source/Runtime/InteractiveToolsFramework/Public/BaseTools/BaseBrushTool.h b/Engine/Source/Runtime/InteractiveToolsFramework/Public/BaseTools/BaseBrushTool.h
--- a/Engine/Source/Runtime/InteractiveToolsFramework/Public/BaseTools/BaseBrushTool.h
+++ b/Engine/Source/Runtime/InteractiveToolsFramework/Public/BaseTools/BaseBrushTool.h
@@ -102,10 +102,15 @@ public:
 
 	virtual void OnPropertyModified(UObject* PropertySet, FProperty* Property) override;
 
-
-	//
+	// IClickDragBehaviorTarget implementation
+	virtual FInputRayHit CanBeginClickDragSequence(const FInputDeviceRay& PressPos) override;
+	virtual void OnClickPress(const FInputDeviceRay& PressPos) override;
+	virtual void OnClickDrag(const FInputDeviceRay& DragPos) override;
+	virtual void OnClickRelease(const FInputDeviceRay& ReleasePos) override;
+	virtual void OnTerminateDragSequence() override;
+	virtual void OnUpdateModifierState(int ModifierID, bool bIsOn) override;
+	
 	// UMeshSurfacePointTool implementation
-	//
 	virtual void OnBeginDrag(const FRay& Ray) override;
 	virtual void OnUpdateDrag(const FRay& Ray) override;
 	virtual void OnEndDrag(const FRay& Ray) override;
@@ -143,6 +148,8 @@ public:
 	virtual double GetCurrentBrushRadius() const { return CurrentBrushRadius; }
 	virtual double GetCurrentBrushRadiusLocal() const { return CurrentBrushRadius * WorldToLocalScale; }
 
+	void SetBrushEnabled(bool bIsEnabled);
+
 protected:
 
 	/**
@@ -169,4 +176,8 @@ protected:
 	virtual void SetupBrushStampIndicator();
 	virtual void UpdateBrushStampIndicator();
 	virtual void ShutdownBrushStampIndicator();
+
+private:
+	bool bEnabled = true;
+	bool bAdjustingBrushSize = false;
 };
