Removed code that should not have made it into last change (25500261) due to shelving incident.

diff --git a/Engine/Source/Runtime/InteractiveToolsFramework/Private/BaseTools/BaseBrushTool.cpp b/Engine/Source/Runtime/InteractiveToolsFramework/Private/BaseTools/BaseBrushTool.cpp
--- a/Engine/Source/Runtime/InteractiveToolsFramework/Private/BaseTools/BaseBrushTool.cpp
+++ b/Engine/Source/Runtime/InteractiveToolsFramework/Private/BaseTools/BaseBrushTool.cpp
@@ -80,31 +80,6 @@ FInputRayHit UBaseBrushTool::CanBeginClickDragSequence(const FInputDeviceRay& Pr
 	return FInputRayHit();
 }
 
-void UBaseBrushTool::OnClickPress(const FInputDeviceRay& PressPos)
-{
-	Super::OnClickPress(PressPos);
-}
-
-void UBaseBrushTool::OnClickDrag(const FInputDeviceRay& DragPos)
-{
-	Super::OnClickDrag(DragPos);
-}
-
-void UBaseBrushTool::OnClickRelease(const FInputDeviceRay& ReleasePos)
-{
-	Super::OnClickRelease(ReleasePos);
-}
-
-void UBaseBrushTool::OnTerminateDragSequence()
-{
-	Super::OnTerminateDragSequence();
-}
-
-void UBaseBrushTool::OnUpdateModifierState(int ModifierID, bool bIsOn)
-{
-	Super::OnUpdateModifierState(ModifierID, bIsOn);
-}
-
 void UBaseBrushTool::IncreaseBrushSizeAction()
 {
 	if (BrushProperties->bSpecifyRadius)
diff --git a/Engine/Source/Runtime/InteractiveToolsFramework/Public/BaseTools/BaseBrushTool.h b/Engine/Source/Runtime/InteractiveToolsFramework/Public/BaseTools/BaseBrushTool.h
--- a/Engine/Source/Runtime/InteractiveToolsFramework/Public/BaseTools/BaseBrushTool.h
+++ b/Engine/Source/Runtime/InteractiveToolsFramework/Public/BaseTools/BaseBrushTool.h
@@ -104,11 +104,6 @@ public:
 
 	// IClickDragBehaviorTarget implementation
 	virtual FInputRayHit CanBeginClickDragSequence(const FInputDeviceRay& PressPos) override;
-	virtual void OnClickPress(const FInputDeviceRay& PressPos) override;
-	virtual void OnClickDrag(const FInputDeviceRay& DragPos) override;
-	virtual void OnClickRelease(const FInputDeviceRay& ReleasePos) override;
-	virtual void OnTerminateDragSequence() override;
-	virtual void OnUpdateModifierState(int ModifierID, bool bIsOn) override;
 	
 	// UMeshSurfacePointTool implementation
 	virtual void OnBeginDrag(const FRay& Ray) override;
