Adding ability to adjust brush size and strength while click-dragging with B hotkey pressed.

diff --git a/Engine/Source/Runtime/InteractiveToolsFramework/Private/BaseTools/BaseBrushTool.cpp b/Engine/Source/Runtime/InteractiveToolsFramework/Private/BaseTools/BaseBrushTool.cpp
--- a/Engine/Source/Runtime/InteractiveToolsFramework/Private/BaseTools/BaseBrushTool.cpp
+++ b/Engine/Source/Runtime/InteractiveToolsFramework/Private/BaseTools/BaseBrushTool.cpp
@@ -1,10 +1,14 @@
 // Copyright Epic Games, Inc. All Rights Reserved.
 
 #include "BaseTools/BaseBrushTool.h"
+
+#include "Engine/Engine.h"
+#include "CanvasItem.h"
+#include "CanvasTypes.h"
 #include "InteractiveToolManager.h"
 #include "InteractiveGizmoManager.h"
-#include "ToolBuilderUtil.h"
 #include "InteractiveTool.h"
+#include "Android/AndroidPlatformApplicationMisc.h"
 #include "BaseBehaviors/ClickDragBehavior.h"
 #include "BaseGizmos/BrushStampIndicator.h"
 
@@ -14,7 +18,6 @@
 #define LOCTEXT_NAMESPACE "UBaseBrushTool"
 
 
-
 UBrushBaseProperties::UBrushBaseProperties()
 {
 	BrushSize = 0.25f;
@@ -24,13 +27,159 @@ UBrushBaseProperties::UBrushBaseProperties()
 	BrushFalloffAmount = 1.0f;
 }
 
+void UBrushAdjusterInputBehavior::Initialize(UBaseBrushTool* InBrushTool)
+{
+	BrushTool = InBrushTool;
+}
+
+void UBrushAdjusterInputBehavior::DrawHUD(FCanvas* Canvas, IToolsContextRenderAPI* RenderAPI)
+{
+	if (!bAdjustingBrush)
+	{
+		return;
+	}
+
+	FText BrushAdjustmentMessage;
+	if (bAdjustingHorizontally)
+	{
+		BrushAdjustmentMessage = FText::Format(LOCTEXT("AdjustRadius", "Radius: {0}"), FText::AsNumber(BrushTool->BrushProperties->BrushRadius));
+	}
+	else
+	{
+		BrushAdjustmentMessage = FText::Format(LOCTEXT("AdjustStrength", "Strength: {0}"), FText::AsNumber(BrushTool->BrushProperties->BrushStrength));
+	}
+
+	FCanvasTextItem TextItem(BrushOrigin, BrushAdjustmentMessage, GEngine->GetMediumFont(), FLinearColor::White);
+	TextItem.EnableShadow(FLinearColor::Black);
+	Canvas->DrawItem(TextItem);
+}
+
+void UBrushAdjusterInputBehavior::OnDragStart(FVector2D InScreenPosition)
+{
+	constexpr bool bAdjustHorizontal = true;
+	BrushOrigin = InScreenPosition;
+	ResetAdjustmentOrigin(InScreenPosition, bAdjustHorizontal);
+}
+
+void UBrushAdjusterInputBehavior::ResetAdjustmentOrigin(FVector2D InScreenPosition, bool bHorizontalAdjust)
+{
+	bAdjustingHorizontally = bHorizontalAdjust;
+	AdjustmentOrigin = InScreenPosition;
+	StartBrushRadius = BrushTool->BrushProperties->BrushRadius;
+	StartBrushStrength = BrushTool->BrushProperties->BrushStrength;
+}
+
+void UBrushAdjusterInputBehavior::OnDragUpdate(FVector2D InScreenPosition)
+{
+	if (!bAdjustingBrush)
+	{
+		return;
+	}
+
+	// calculate screen space cursor delta relative to adjustment origin
+	const float HorizontalDelta = InScreenPosition.X - AdjustmentOrigin.X;
+	const float VerticalDelta = InScreenPosition.Y - AdjustmentOrigin.Y;
+	const float HorizDeltaMag = FMath::Abs(HorizontalDelta);
+	const float VertDeltaMag = FMath::Abs(VerticalDelta);
+
+	if (bAdjustingHorizontally && HorizDeltaMag < VertDeltaMag)
+	{
+		// switch to adjusting vertically and re-center adjustment origin
+		ResetAdjustmentOrigin(InScreenPosition, false);
+	}
+	else if (!bAdjustingHorizontally && VertDeltaMag < HorizDeltaMag)
+	{
+		// switch to adjusting horizontally and re-center adjustment origin
+		ResetAdjustmentOrigin(InScreenPosition, true);
+	}
+
+	// scale for consistent screen space speed on varying monitor DPI
+	// (takes device coordinates as input because multi-monitor setups may have different DPI)
+	const float DPIScale = FPlatformApplicationMisc::GetDPIScaleFactorAtPoint(InScreenPosition.X, InScreenPosition.Y);
+	
+	// apply directional adjustments
+	if (bAdjustingHorizontally)
+	{
+		// adjust brush size based on horizontal mouse drag
+		float NewRadius = StartBrushRadius + HorizontalDelta * (SizeAdjustSpeed * DPIScale);
+		NewRadius = FMath::Max(NewRadius, 0.01f);
+		BrushTool->BrushProperties->BrushRadius = NewRadius;
+		BrushTool->LastBrushStamp.Radius = NewRadius;
+	}
+	else
+	{
+		// adjust brush strength based on vertical mouse drag
+		float NewStrength = StartBrushStrength + VerticalDelta * -(StrengthAdjustSpeed * DPIScale);
+		NewStrength = FMath::Min(1.0f,FMath::Max(NewStrength, 0.f));
+		BrushTool->BrushProperties->BrushStrength = NewStrength;
+	}	
+}
+
+EInputDevices UBrushAdjusterInputBehavior::GetSupportedDevices()
+{
+	return EInputDevices::Keyboard;
+}
+
+bool UBrushAdjusterInputBehavior::IsPressed(const FInputDeviceState& Input)
+{
+	if (Input.IsFromDevice(EInputDevices::Keyboard))
+	{
+		ActiveDevice = EInputDevices::Keyboard;
+		return Input.Keyboard.ActiveKey.Button == EKeys::B && Input.Keyboard.ActiveKey.bDown;
+	}
+	
+	return false;
+}
+
+bool UBrushAdjusterInputBehavior::IsReleased(const FInputDeviceState& Input)
+{
+	if (Input.IsFromDevice(EInputDevices::Keyboard))
+	{
+		return Input.Keyboard.ActiveKey.Button == EKeys::B && Input.Keyboard.ActiveKey.bReleased;
+	}
+	
+	return false;
+}
+
+FInputCaptureRequest UBrushAdjusterInputBehavior::WantsCapture(const FInputDeviceState& Input)
+{
+	if (IsPressed(Input))
+	{
+		return FInputCaptureRequest::Begin(this, EInputCaptureSide::Any, 0.f);
+	}
+
+	return FInputCaptureRequest::Ignore();
+}
+
+FInputCaptureUpdate UBrushAdjusterInputBehavior::BeginCapture(const FInputDeviceState& Input, EInputCaptureSide Side)
+{
+	bAdjustingBrush = true;
+	return FInputCaptureUpdate::Begin(this, EInputCaptureSide::Any);
+}
+
+FInputCaptureUpdate UBrushAdjusterInputBehavior::UpdateCapture(
+	const FInputDeviceState& Input,
+	const FInputCaptureData& Data)
+{
+	if (IsReleased(Input))
+	{
+		bAdjustingBrush = false;
+		return FInputCaptureUpdate::End();
+	}
+
+	return FInputCaptureUpdate::Continue();
+}
+
+void UBrushAdjusterInputBehavior::ForceEndCapture(const FInputCaptureData& data)
+{
+	bAdjustingBrush = false;
+}
+
 UBaseBrushTool::UBaseBrushTool()
 {
 	PropertyClass = UBrushBaseProperties::StaticClass();
 }
 
-static const int UBaseBrushTool_BrushSizeModifier = 0;
-
 void UBaseBrushTool::Setup()
 {
 	UMeshSurfacePointTool::Setup();
@@ -45,14 +194,10 @@ void UBaseBrushTool::Setup()
 
 	SetupBrushStampIndicator();
 
-	// add input behavior to click-drag while holding hotkey to adjust brush size
-	UClickDragInputBehavior* DragBehavior = NewObject<UClickDragInputBehavior>();
-	DragBehavior->Modifiers.RegisterModifier(UBaseBrushTool_BrushSizeModifier, [](const FInputDeviceState& Input)
-	{
-		return (Input.Keyboard.ActiveKey.Button == EKeys::B) && Input.Keyboard.ActiveKey.bDown;
-	});
-	DragBehavior->Initialize(this);
-	AddInputBehavior(DragBehavior);
+	// add input behavior to click-drag while holding hotkey to adjust brush size and strength
+	BrushAdjusterBehavior = NewObject<UBrushAdjusterInputBehavior>(this);
+	BrushAdjusterBehavior->Initialize(this);
+	AddInputBehavior(BrushAdjusterBehavior.Get());
 }
 
 
@@ -72,12 +217,32 @@ void UBaseBrushTool::OnPropertyModified(UObject* PropertySet, FProperty* Propert
 
 FInputRayHit UBaseBrushTool::CanBeginClickDragSequence(const FInputDeviceRay& PressPos)
 {
-	if (bEnabled)
+	if (!bEnabled)
+	{
+		// no hit
+		return FInputRayHit();
+	}
+	
+	if (BrushAdjusterBehavior->IsBrushBeingAdjusted())
 	{
-		return Super::CanBeginClickDragSequence(PressPos);	
+		// fake screen hit
+		return FInputRayHit(0.f);
 	}
 
-	return FInputRayHit();
+	// hit-test the tool target
+	return Super::CanBeginClickDragSequence(PressPos);	
+}
+
+void UBaseBrushTool::OnClickPress(const FInputDeviceRay& PressPos)
+{
+	Super::OnClickPress(PressPos);
+	BrushAdjusterBehavior->OnDragStart(PressPos.ScreenPosition);
+}
+
+void UBaseBrushTool::OnClickDrag(const FInputDeviceRay& DragPos)
+{
+	Super::OnClickDrag(DragPos);
+	BrushAdjusterBehavior->OnDragUpdate(DragPos.ScreenPosition);
 }
 
 void UBaseBrushTool::IncreaseBrushSizeAction()
@@ -224,6 +389,12 @@ void UBaseBrushTool::RecalculateBrushRadius()
 
 void UBaseBrushTool::OnBeginDrag(const FRay& Ray)
 {
+	if (BrushAdjusterBehavior->IsBrushBeingAdjusted())
+	{
+		bInBrushStroke = false;
+		return;
+	}
+	
 	FHitResult OutHit;
 	if (HitTest(Ray, OutHit))
 	{
@@ -238,6 +409,13 @@ void UBaseBrushTool::OnBeginDrag(const FRay& Ray)
 
 void UBaseBrushTool::OnUpdateDrag(const FRay& Ray)
 {
+	if (BrushAdjusterBehavior->IsBrushBeingAdjusted())
+	{
+		RecalculateBrushRadius();
+		NotifyOfPropertyChangeByTool(BrushProperties);
+		return;
+	}
+	
 	FHitResult OutHit;
 	if (HitTest(Ray, OutHit))
 	{
@@ -256,6 +434,11 @@ void UBaseBrushTool::OnEndDrag(const FRay& Ray)
 
 bool UBaseBrushTool::OnUpdateHover(const FInputDeviceRay& DevicePos)
 {
+	if (BrushAdjusterBehavior->IsBrushBeingAdjusted())
+	{
+		return true;
+	}
+	
 	FHitResult OutHit;
 	if (HitTest(DevicePos.WorldRay, OutHit))
 	{
@@ -277,6 +460,12 @@ void UBaseBrushTool::Render(IToolsContextRenderAPI* RenderAPI)
 	}
 }
 
+void UBaseBrushTool::DrawHUD(FCanvas* Canvas, IToolsContextRenderAPI* RenderAPI)
+{
+	Super::DrawHUD(Canvas, RenderAPI);
+	BrushAdjusterBehavior->DrawHUD(Canvas, RenderAPI);
+}
+
 const FString BaseBrushIndicatorGizmoType = TEXT("BrushIndicatorGizmoType");
 
 void UBaseBrushTool::SetupBrushStampIndicator()
@@ -293,6 +482,7 @@ void UBaseBrushTool::UpdateBrushStampIndicator()
 {
 	if (BrushStampIndicator)
 	{
+		BrushStampIndicator->LineColor = BrushAdjusterBehavior->IsBrushBeingAdjusted() ? FLinearColor::White : FLinearColor::Green;
 		BrushStampIndicator->Update(LastBrushStamp.Radius, LastBrushStamp.WorldPosition, LastBrushStamp.WorldNormal, LastBrushStamp.Falloff);
 	}
 }
diff --git a/Engine/Source/Runtime/InteractiveToolsFramework/Public/BaseTools/BaseBrushTool.h b/Engine/Source/Runtime/InteractiveToolsFramework/Public/BaseTools/BaseBrushTool.h
--- a/Engine/Source/Runtime/InteractiveToolsFramework/Public/BaseTools/BaseBrushTool.h
+++ b/Engine/Source/Runtime/InteractiveToolsFramework/Public/BaseTools/BaseBrushTool.h
@@ -75,7 +75,69 @@ struct INTERACTIVETOOLSFRAMEWORK_API FBrushStampData
 	float Falloff;
 };
 
+/**
+ * A behavior that captures a keyboard hotkey to enter a sub-mode "bAdjustingBrush" while the key is pressed.
+ * The target tool must call OnDragStart() and OnDragUpdate() to feed the screen coordinates of the mouse for the duration
+ * of the behavior. And use "GetIsBrushBeingAdjusted" to pause/disable the brush motion while it is being adjusted.
+ *
+ * OnDragStart() defines the starting location of an adjustment
+ * OnDragUpdate() adjusts the brush strength and radius based on the magnitude of the screen coordinate delta in the
+ * vertical and horizontal directions respectively.
+ */
+UCLASS()
+class UBrushAdjusterInputBehavior : public UAnyButtonInputBehavior
+{
+	GENERATED_BODY()
+
+public:
+	
+	void Initialize(UBaseBrushTool* InBrushTool);
+
+	void DrawHUD(FCanvas* Canvas, IToolsContextRenderAPI* RenderAPI);
+	
+	void SetStrengthAdjustSpeed(float InStrengthSpeed) { StrengthAdjustSpeed = InStrengthSpeed; };
+	void SetSizeAdjustSpeed(float InSizeSpeed) { SizeAdjustSpeed = InSizeSpeed; };
+	
+	void OnDragStart(FVector2D InScreenPosition);
+	void OnDragUpdate(FVector2D InScreenPosition);
+	
+	bool IsBrushBeingAdjusted() const { return bAdjustingBrush; };
+
+	// UAnyButtonInputBehavior
+	virtual EInputDevices GetSupportedDevices() override;
+	virtual bool IsPressed(const FInputDeviceState& Input) override;
+	virtual bool IsReleased(const FInputDeviceState& Input) override;
+	virtual FInputCaptureRequest WantsCapture(const FInputDeviceState& Input) override;
+	virtual FInputCaptureUpdate BeginCapture(const FInputDeviceState& Input, EInputCaptureSide eSide) override;
+	virtual FInputCaptureUpdate UpdateCapture(const FInputDeviceState& Input, const FInputCaptureData& data) override;
+	virtual void ForceEndCapture(const FInputCaptureData& data) override;
+	// UAnyButtonInputBehavior end
 
+private:
+
+	void ResetAdjustmentOrigin(FVector2D InScreenPosition, bool bHorizontalAdjust);
+
+	// true for the duration of the behavior (while hotkey is down)
+	bool bAdjustingBrush = false;
+	// screen coordinate when hotkey pressed (or when re-centered if changing direction)
+	FVector2D BrushOrigin;
+	// screen coordinate when hotkey pressed (or when re-centered if changing direction)
+	FVector2D AdjustmentOrigin;
+	// the radius of the brush when we started a drag
+	float StartBrushRadius;
+	// the strength of the brush when we started a drag
+	float StartBrushStrength;
+	// supports separate horizontal / vertical adjustments (switching dynamically based on magnitude)
+	bool bAdjustingHorizontally = true;
+
+	// the speed (centimeters per unit of screen coordinate) to adjust brush size when dragging
+	float SizeAdjustSpeed = 0.04f; // sensible default based on 1080p monitor
+	// the speed (in strength per unit of screen coordinate) to adjust brush strength when dragging
+	float StrengthAdjustSpeed = 0.005f; // sensible default based on 1080p monitor
+
+	// the target brush tool to adjust
+	UBaseBrushTool* BrushTool;
+};
 
 /**
  * UBaseBrushTool implements standard brush-style functionality for an InteractiveTool.
@@ -84,6 +146,7 @@ struct INTERACTIVETOOLSFRAMEWORK_API FBrushStampData
  *   2) brush indicator visualization
  *   3) tracking of last brush stamp location via .LastBrushStamp FProperty
  *   4) status of brush stroke via .bInBrushStroke FProperty
+ *   5) "B" hotkey to adjust brush radius / strength by click-dragging in the viewport
  */
 UCLASS()
 class INTERACTIVETOOLSFRAMEWORK_API UBaseBrushTool : public UMeshSurfacePointTool
@@ -95,15 +158,20 @@ public:
 
 	virtual void RegisterActions(FInteractiveToolActionSet& ActionSet) override;
 
+	
+
+	// UInteractiveTool
 	virtual void Setup() override;
 	virtual void Shutdown(EToolShutdownType ShutdownType) override;
-
 	virtual void Render(IToolsContextRenderAPI* RenderAPI) override;
-
+	virtual void DrawHUD( FCanvas* Canvas, IToolsContextRenderAPI* RenderAPI ) override;
+	
 	virtual void OnPropertyModified(UObject* PropertySet, FProperty* Property) override;
-
+	
 	// IClickDragBehaviorTarget implementation
 	virtual FInputRayHit CanBeginClickDragSequence(const FInputDeviceRay& PressPos) override;
+	virtual void OnClickPress(const FInputDeviceRay& PressPos) override;
+	virtual void OnClickDrag(const FInputDeviceRay& DragPos) override;
 	
 	// UMeshSurfacePointTool implementation
 	virtual void OnBeginDrag(const FRay& Ray) override;
@@ -172,7 +240,10 @@ protected:
 	virtual void UpdateBrushStampIndicator();
 	virtual void ShutdownBrushStampIndicator();
 
+	// adjusts size and strength properties while holding hotkey during click/drag
+	TWeakObjectPtr<UBrushAdjusterInputBehavior> BrushAdjusterBehavior;
+
 private:
+	
 	bool bEnabled = true;
-	bool bAdjustingBrushSize = false;
 };
