Minor changes to BaseBrushTool, added ability to disable brush adjustment input, and removed platform specific header.

diff --git a/Engine/Source/Runtime/InteractiveToolsFramework/Private/BaseTools/BaseBrushTool.cpp b/Engine/Source/Runtime/InteractiveToolsFramework/Private/BaseTools/BaseBrushTool.cpp
--- a/Engine/Source/Runtime/InteractiveToolsFramework/Private/BaseTools/BaseBrushTool.cpp
+++ b/Engine/Source/Runtime/InteractiveToolsFramework/Private/BaseTools/BaseBrushTool.cpp
@@ -8,7 +8,7 @@
 #include "InteractiveToolManager.h"
 #include "InteractiveGizmoManager.h"
 #include "InteractiveTool.h"
-#include "Android/AndroidPlatformApplicationMisc.h"
+#include "GenericPlatform/GenericPlatformApplicationMisc.h"
 #include "BaseBehaviors/ClickDragBehavior.h"
 #include "BaseGizmos/BrushStampIndicator.h"
 
@@ -95,7 +95,7 @@ void UBrushAdjusterInputBehavior::OnDragUpdate(FVector2D InScreenPosition)
 
 	// scale for consistent screen space speed on varying monitor DPI
 	// (takes device coordinates as input because multi-monitor setups may have different DPI)
-	const float DPIScale = FPlatformApplicationMisc::GetDPIScaleFactorAtPoint(InScreenPosition.X, InScreenPosition.Y);
+	const float DPIScale = FGenericPlatformApplicationMisc::GetDPIScaleFactorAtPoint(InScreenPosition.X, InScreenPosition.Y);
 	
 	// apply directional adjustments
 	if (bAdjustingHorizontally)
@@ -195,9 +195,12 @@ void UBaseBrushTool::Setup()
 	SetupBrushStampIndicator();
 
 	// add input behavior to click-drag while holding hotkey to adjust brush size and strength
-	BrushAdjusterBehavior = NewObject<UBrushAdjusterInputBehavior>(this);
-	BrushAdjusterBehavior->Initialize(this);
-	AddInputBehavior(BrushAdjusterBehavior.Get());
+	if (SupportsBrushAdjustmentInput())
+	{
+		BrushAdjusterBehavior = NewObject<UBrushAdjusterInputBehavior>(this);
+		BrushAdjusterBehavior->Initialize(this);
+		AddInputBehavior(BrushAdjusterBehavior.Get());	
+	}
 }
 
 
@@ -223,7 +226,7 @@ FInputRayHit UBaseBrushTool::CanBeginClickDragSequence(const FInputDeviceRay& Pr
 		return FInputRayHit();
 	}
 	
-	if (BrushAdjusterBehavior->IsBrushBeingAdjusted())
+	if (BrushAdjusterBehavior.IsValid() && BrushAdjusterBehavior->IsBrushBeingAdjusted())
 	{
 		// fake screen hit
 		return FInputRayHit(0.f);
@@ -236,13 +239,19 @@ FInputRayHit UBaseBrushTool::CanBeginClickDragSequence(const FInputDeviceRay& Pr
 void UBaseBrushTool::OnClickPress(const FInputDeviceRay& PressPos)
 {
 	Super::OnClickPress(PressPos);
-	BrushAdjusterBehavior->OnDragStart(PressPos.ScreenPosition);
+	if (BrushAdjusterBehavior.IsValid())
+	{
+		BrushAdjusterBehavior->OnDragStart(PressPos.ScreenPosition);
+	}
 }
 
 void UBaseBrushTool::OnClickDrag(const FInputDeviceRay& DragPos)
 {
 	Super::OnClickDrag(DragPos);
-	BrushAdjusterBehavior->OnDragUpdate(DragPos.ScreenPosition);
+	if (BrushAdjusterBehavior.IsValid())
+	{
+		BrushAdjusterBehavior->OnDragUpdate(DragPos.ScreenPosition);
+	}
 }
 
 void UBaseBrushTool::IncreaseBrushSizeAction()
@@ -389,7 +398,7 @@ void UBaseBrushTool::RecalculateBrushRadius()
 
 void UBaseBrushTool::OnBeginDrag(const FRay& Ray)
 {
-	if (BrushAdjusterBehavior->IsBrushBeingAdjusted())
+	if (BrushAdjusterBehavior.IsValid() && BrushAdjusterBehavior->IsBrushBeingAdjusted())
 	{
 		bInBrushStroke = false;
 		return;
@@ -409,7 +418,7 @@ void UBaseBrushTool::OnBeginDrag(const FRay& Ray)
 
 void UBaseBrushTool::OnUpdateDrag(const FRay& Ray)
 {
-	if (BrushAdjusterBehavior->IsBrushBeingAdjusted())
+	if (BrushAdjusterBehavior.IsValid() && BrushAdjusterBehavior->IsBrushBeingAdjusted())
 	{
 		RecalculateBrushRadius();
 		NotifyOfPropertyChangeByTool(BrushProperties);
@@ -434,7 +443,7 @@ void UBaseBrushTool::OnEndDrag(const FRay& Ray)
 
 bool UBaseBrushTool::OnUpdateHover(const FInputDeviceRay& DevicePos)
 {
-	if (BrushAdjusterBehavior->IsBrushBeingAdjusted())
+	if (BrushAdjusterBehavior.IsValid() && BrushAdjusterBehavior->IsBrushBeingAdjusted())
 	{
 		return true;
 	}
@@ -463,7 +472,10 @@ void UBaseBrushTool::Render(IToolsContextRenderAPI* RenderAPI)
 void UBaseBrushTool::DrawHUD(FCanvas* Canvas, IToolsContextRenderAPI* RenderAPI)
 {
 	Super::DrawHUD(Canvas, RenderAPI);
-	BrushAdjusterBehavior->DrawHUD(Canvas, RenderAPI);
+	if (BrushAdjusterBehavior.IsValid())
+	{
+		BrushAdjusterBehavior->DrawHUD(Canvas, RenderAPI);	
+	}
 }
 
 const FString BaseBrushIndicatorGizmoType = TEXT("BrushIndicatorGizmoType");
@@ -482,7 +494,11 @@ void UBaseBrushTool::UpdateBrushStampIndicator()
 {
 	if (BrushStampIndicator)
 	{
-		BrushStampIndicator->LineColor = BrushAdjusterBehavior->IsBrushBeingAdjusted() ? FLinearColor::White : FLinearColor::Green;
+		if (BrushAdjusterBehavior.IsValid())
+		{
+			BrushStampIndicator->LineColor = BrushAdjusterBehavior->IsBrushBeingAdjusted() ? FLinearColor::White : FLinearColor::Green;	
+		}
+		
 		BrushStampIndicator->Update(LastBrushStamp.Radius, LastBrushStamp.WorldPosition, LastBrushStamp.WorldNormal, LastBrushStamp.Falloff);
 	}
 }
diff --git a/Engine/Source/Runtime/InteractiveToolsFramework/Public/BaseTools/BaseBrushTool.h b/Engine/Source/Runtime/InteractiveToolsFramework/Public/BaseTools/BaseBrushTool.h
--- a/Engine/Source/Runtime/InteractiveToolsFramework/Public/BaseTools/BaseBrushTool.h
+++ b/Engine/Source/Runtime/InteractiveToolsFramework/Public/BaseTools/BaseBrushTool.h
@@ -158,8 +158,6 @@ public:
 
 	virtual void RegisterActions(FInteractiveToolActionSet& ActionSet) override;
 
-	
-
 	// UInteractiveTool
 	virtual void Setup() override;
 	virtual void Shutdown(EToolShutdownType ShutdownType) override;
@@ -213,6 +211,9 @@ public:
 
 	void SetBrushEnabled(bool bIsEnabled);
 
+	// return false to disable the hotkey for adjusting the brush radius and strength in the viewport
+	virtual bool SupportsBrushAdjustmentInput() { return true; }; 
+
 protected:
 
 	/**
