[Engine]

diff --git a/Engine/Source/Runtime/InteractiveToolsFramework/Public/BaseBehaviors/MouseHoverBehavior.h b/Engine/Source/Runtime/InteractiveToolsFramework/Public/BaseBehaviors/MouseHoverBehavior.h
--- a/Engine/Source/Runtime/InteractiveToolsFramework/Public/BaseBehaviors/MouseHoverBehavior.h
+++ b/Engine/Source/Runtime/InteractiveToolsFramework/Public/BaseBehaviors/MouseHoverBehavior.h
@@ -22,30 +22,30 @@ class UObject;
  * Trivial InputBehavior that forwards InputBehavior hover events to a Target object via
  * the IHoverBehaviorTarget interface.
  */
-UCLASS()
-class INTERACTIVETOOLSFRAMEWORK_API UMouseHoverBehavior : public UInputBehavior
+UCLASS(MinimalAPI)
+class UMouseHoverBehavior : public UInputBehavior
 {
 	GENERATED_BODY()
 
 public:
-	UMouseHoverBehavior();
+	INTERACTIVETOOLSFRAMEWORK_API UMouseHoverBehavior();
 
 	/**
 	 * The modifier set for this behavior
 	 */
 	FInputBehaviorModifierStates Modifiers;
 
-	virtual void Initialize(IHoverBehaviorTarget* Target);
+	INTERACTIVETOOLSFRAMEWORK_API virtual void Initialize(IHoverBehaviorTarget* Target);
 
 	// UInputBehavior hover implementation
 
-	virtual EInputDevices GetSupportedDevices() override;
+	INTERACTIVETOOLSFRAMEWORK_API virtual EInputDevices GetSupportedDevices() override;
 
-	virtual bool WantsHoverEvents() override;
-	virtual FInputCaptureRequest WantsHoverCapture(const FInputDeviceState& InputState) override;
-	virtual FInputCaptureUpdate BeginHoverCapture(const FInputDeviceState& InputState, EInputCaptureSide eSide) override;
-	virtual FInputCaptureUpdate UpdateHoverCapture(const FInputDeviceState& InputState) override;
-	virtual void EndHoverCapture() override;
+	INTERACTIVETOOLSFRAMEWORK_API virtual bool WantsHoverEvents() override;
+	INTERACTIVETOOLSFRAMEWORK_API virtual FInputCaptureRequest WantsHoverCapture(const FInputDeviceState& InputState) override;
+	INTERACTIVETOOLSFRAMEWORK_API virtual FInputCaptureUpdate BeginHoverCapture(const FInputDeviceState& InputState, EInputCaptureSide eSide) override;
+	INTERACTIVETOOLSFRAMEWORK_API virtual FInputCaptureUpdate UpdateHoverCapture(const FInputDeviceState& InputState) override;
+	INTERACTIVETOOLSFRAMEWORK_API virtual void EndHoverCapture() override;
 
 protected:
 	IHoverBehaviorTarget* Target;
@@ -57,8 +57,8 @@ protected:
  * of local lambda functions. To use/customize this class, the client replaces the lambda functions with their own.
  * This avoids having to create a separate IHoverBehaviorTarget implementation for trivial use-cases.
  */
-UCLASS()
-class INTERACTIVETOOLSFRAMEWORK_API ULocalMouseHoverBehavior : public UMouseHoverBehavior, public IHoverBehaviorTarget
+UCLASS(MinimalAPI)
+class ULocalMouseHoverBehavior : public UMouseHoverBehavior, public IHoverBehaviorTarget
 {
 	GENERATED_BODY()
 protected:
diff --git a/Engine/Source/Runtime/InteractiveToolsFramework/Public/BaseBehaviors/MouseWheelBehavior.h b/Engine/Source/Runtime/InteractiveToolsFramework/Public/BaseBehaviors/MouseWheelBehavior.h
--- a/Engine/Source/Runtime/InteractiveToolsFramework/Public/BaseBehaviors/MouseWheelBehavior.h
+++ b/Engine/Source/Runtime/InteractiveToolsFramework/Public/BaseBehaviors/MouseWheelBehavior.h
@@ -16,19 +16,19 @@ class IMouseWheelBehaviorTarget;
 class UObject;
 struct FInputDeviceState;
 
-UCLASS()
-class INTERACTIVETOOLSFRAMEWORK_API UMouseWheelInputBehavior : public UAnyButtonInputBehavior
+UCLASS(MinimalAPI)
+class UMouseWheelInputBehavior : public UAnyButtonInputBehavior
 {
 	GENERATED_BODY()
 
 public:
-	UMouseWheelInputBehavior();
+	INTERACTIVETOOLSFRAMEWORK_API UMouseWheelInputBehavior();
 
 	/**
 	 * Initialize this behavior with the given Target
 	 * @param Target implementor of hit-test and on-clicked functions
 	 */
-	virtual void Initialize(IMouseWheelBehaviorTarget* Target);
+	INTERACTIVETOOLSFRAMEWORK_API virtual void Initialize(IMouseWheelBehaviorTarget* Target);
 
 
 	/**
@@ -38,10 +38,10 @@ public:
 
 
 	// UInputBehavior implementation
-	virtual FInputCaptureRequest WantsCapture(const FInputDeviceState& Input) override;
-	virtual FInputCaptureUpdate BeginCapture(const FInputDeviceState& Input, EInputCaptureSide eSide) override;
-	virtual FInputCaptureUpdate UpdateCapture(const FInputDeviceState& Input, const FInputCaptureData& Data) override;
-	virtual void ForceEndCapture(const FInputCaptureData& Data) override;
+	INTERACTIVETOOLSFRAMEWORK_API virtual FInputCaptureRequest WantsCapture(const FInputDeviceState& Input) override;
+	INTERACTIVETOOLSFRAMEWORK_API virtual FInputCaptureUpdate BeginCapture(const FInputDeviceState& Input, EInputCaptureSide eSide) override;
+	INTERACTIVETOOLSFRAMEWORK_API virtual FInputCaptureUpdate UpdateCapture(const FInputDeviceState& Input, const FInputCaptureData& Data) override;
+	INTERACTIVETOOLSFRAMEWORK_API virtual void ForceEndCapture(const FInputCaptureData& Data) override;
 
 	/**
 	 * The modifier set for this behavior
@@ -54,4 +54,4 @@ protected:
 	 */
 	IMouseWheelBehaviorTarget* Target;
 
-};
\ No newline at end of file
+};
diff --git a/Engine/Source/Runtime/InteractiveToolsFramework/Public/BaseBehaviors/MultiClickSequenceInputBehavior.h b/Engine/Source/Runtime/InteractiveToolsFramework/Public/BaseBehaviors/MultiClickSequenceInputBehavior.h
--- a/Engine/Source/Runtime/InteractiveToolsFramework/Public/BaseBehaviors/MultiClickSequenceInputBehavior.h
+++ b/Engine/Source/Runtime/InteractiveToolsFramework/Public/BaseBehaviors/MultiClickSequenceInputBehavior.h
@@ -34,13 +34,13 @@ struct FInputDeviceState;
  * @todo it may be better to implement this as multiple captures, and use hover callbacks to 
  * do the between-capture previews. holding capture across mouse release is not ideal.
  */
-UCLASS()
-class INTERACTIVETOOLSFRAMEWORK_API UMultiClickSequenceInputBehavior : public UAnyButtonInputBehavior
+UCLASS(MinimalAPI)
+class UMultiClickSequenceInputBehavior : public UAnyButtonInputBehavior
 {
 	GENERATED_BODY()
 
 public:
-	UMultiClickSequenceInputBehavior();
+	INTERACTIVETOOLSFRAMEWORK_API UMultiClickSequenceInputBehavior();
 
 	/**
 	 * The modifier set for this behavior
@@ -51,7 +51,7 @@ public:
 	 * Initialize this behavior with the given Target
 	 * @param Target implementor of hit-test and on-clicked functions
 	 */
-	virtual void Initialize(IClickSequenceBehaviorTarget* Target);
+	INTERACTIVETOOLSFRAMEWORK_API virtual void Initialize(IClickSequenceBehaviorTarget* Target);
 
 
 	/**
@@ -62,16 +62,16 @@ public:
 
 	// UInputBehavior implementation
 
-	virtual FInputCaptureRequest WantsCapture(const FInputDeviceState& Input) override;
-	virtual FInputCaptureUpdate BeginCapture(const FInputDeviceState& Input, EInputCaptureSide eSide) override;
-	virtual FInputCaptureUpdate UpdateCapture(const FInputDeviceState& Input, const FInputCaptureData& Data) override;
-	virtual void ForceEndCapture(const FInputCaptureData& Data) override;
+	INTERACTIVETOOLSFRAMEWORK_API virtual FInputCaptureRequest WantsCapture(const FInputDeviceState& Input) override;
+	INTERACTIVETOOLSFRAMEWORK_API virtual FInputCaptureUpdate BeginCapture(const FInputDeviceState& Input, EInputCaptureSide eSide) override;
+	INTERACTIVETOOLSFRAMEWORK_API virtual FInputCaptureUpdate UpdateCapture(const FInputDeviceState& Input, const FInputCaptureData& Data) override;
+	INTERACTIVETOOLSFRAMEWORK_API virtual void ForceEndCapture(const FInputCaptureData& Data) override;
 
-	virtual bool WantsHoverEvents() override;
-	virtual FInputCaptureRequest WantsHoverCapture(const FInputDeviceState& InputState) override;
-	virtual FInputCaptureUpdate BeginHoverCapture(const FInputDeviceState& InputState, EInputCaptureSide eSide) override;
-	virtual FInputCaptureUpdate UpdateHoverCapture(const FInputDeviceState& InputState) override;
-	virtual void EndHoverCapture() override;
+	INTERACTIVETOOLSFRAMEWORK_API virtual bool WantsHoverEvents() override;
+	INTERACTIVETOOLSFRAMEWORK_API virtual FInputCaptureRequest WantsHoverCapture(const FInputDeviceState& InputState) override;
+	INTERACTIVETOOLSFRAMEWORK_API virtual FInputCaptureUpdate BeginHoverCapture(const FInputDeviceState& InputState, EInputCaptureSide eSide) override;
+	INTERACTIVETOOLSFRAMEWORK_API virtual FInputCaptureUpdate UpdateHoverCapture(const FInputDeviceState& InputState) override;
+	INTERACTIVETOOLSFRAMEWORK_API virtual void EndHoverCapture() override;
 
 
 public:
@@ -87,4 +87,4 @@ protected:
 	IClickSequenceBehaviorTarget* Target;
 
 	ESequenceState State = ESequenceState::NotStarted;
-};
\ No newline at end of file
+};
diff --git a/Engine/Source/Runtime/InteractiveToolsFramework/Public/BaseBehaviors/SingleClickBehavior.h b/Engine/Source/Runtime/InteractiveToolsFramework/Public/BaseBehaviors/SingleClickBehavior.h
--- a/Engine/Source/Runtime/InteractiveToolsFramework/Public/BaseBehaviors/SingleClickBehavior.h
+++ b/Engine/Source/Runtime/InteractiveToolsFramework/Public/BaseBehaviors/SingleClickBehavior.h
@@ -29,19 +29,19 @@ class UObject;
  *    
  * The hit-test and on-clicked functions are provided by a IClickBehaviorTarget instance.
  */
-UCLASS()
-class INTERACTIVETOOLSFRAMEWORK_API USingleClickInputBehavior : public UAnyButtonInputBehavior
+UCLASS(MinimalAPI)
+class USingleClickInputBehavior : public UAnyButtonInputBehavior
 {
 	GENERATED_BODY()
 
 public:
-	USingleClickInputBehavior();
+	INTERACTIVETOOLSFRAMEWORK_API USingleClickInputBehavior();
 
 	/**
 	 * Initialize this behavior with the given Target
 	 * @param Target implementor of hit-test and on-clicked functions
 	 */
-	virtual void Initialize(IClickBehaviorTarget* Target);
+	INTERACTIVETOOLSFRAMEWORK_API virtual void Initialize(IClickBehaviorTarget* Target);
 
 
 	/**
@@ -52,10 +52,10 @@ public:
 
 	// UInputBehavior implementation
 
-	virtual FInputCaptureRequest WantsCapture(const FInputDeviceState& Input) override;
-	virtual FInputCaptureUpdate BeginCapture(const FInputDeviceState& Input, EInputCaptureSide eSide) override;
-	virtual FInputCaptureUpdate UpdateCapture(const FInputDeviceState& Input, const FInputCaptureData& Data) override;
-	virtual void ForceEndCapture(const FInputCaptureData& Data) override;
+	INTERACTIVETOOLSFRAMEWORK_API virtual FInputCaptureRequest WantsCapture(const FInputDeviceState& Input) override;
+	INTERACTIVETOOLSFRAMEWORK_API virtual FInputCaptureUpdate BeginCapture(const FInputDeviceState& Input, EInputCaptureSide eSide) override;
+	INTERACTIVETOOLSFRAMEWORK_API virtual FInputCaptureUpdate UpdateCapture(const FInputDeviceState& Input, const FInputCaptureData& Data) override;
+	INTERACTIVETOOLSFRAMEWORK_API virtual void ForceEndCapture(const FInputCaptureData& Data) override;
 
 
 public:
@@ -76,7 +76,7 @@ protected:
 	/**
 	 * Internal function that forwards click evens to Target::OnClicked, you can customize behavior here
 	 */
-	virtual void Clicked(const FInputDeviceState& Input, const FInputCaptureData& Data);
+	INTERACTIVETOOLSFRAMEWORK_API virtual void Clicked(const FInputDeviceState& Input, const FInputCaptureData& Data);
 };
 
 /**
@@ -84,8 +84,8 @@ protected:
  * of local lambda functions. To use/customize this class, the client replaces the lambda functions with their own.
  * This avoids having to create a separate IClickBehaviorTarget implementation for trivial use-cases.
  */
-UCLASS()
-class INTERACTIVETOOLSFRAMEWORK_API ULocalSingleClickInputBehavior : public USingleClickInputBehavior, public IClickBehaviorTarget
+UCLASS(MinimalAPI)
+class ULocalSingleClickInputBehavior : public USingleClickInputBehavior, public IClickBehaviorTarget
 {
 	GENERATED_BODY()
 protected:
diff --git a/Engine/Source/Runtime/InteractiveToolsFramework/Public/BaseBehaviors/SingleClickOrDragBehavior.h b/Engine/Source/Runtime/InteractiveToolsFramework/Public/BaseBehaviors/SingleClickOrDragBehavior.h
--- a/Engine/Source/Runtime/InteractiveToolsFramework/Public/BaseBehaviors/SingleClickOrDragBehavior.h
+++ b/Engine/Source/Runtime/InteractiveToolsFramework/Public/BaseBehaviors/SingleClickOrDragBehavior.h
@@ -37,8 +37,8 @@ struct FInputDeviceState;
  * The hit-test and on-clicked functions are provided by a IClickBehaviorTarget instance, while an
  * IClickDragBehaviorTarget provides the can-click-drag/begin-drag/update-drag/end-drag functionality.
  */
-UCLASS()
-class INTERACTIVETOOLSFRAMEWORK_API USingleClickOrDragInputBehavior : public UAnyButtonInputBehavior
+UCLASS(MinimalAPI)
+class USingleClickOrDragInputBehavior : public UAnyButtonInputBehavior
 {
 	GENERATED_BODY()
 
@@ -53,12 +53,12 @@ public:
 	/**
 	 * Set the targets for Click and Drag interactions
 	 */
-	virtual void Initialize(IClickBehaviorTarget* ClickTarget, IClickDragBehaviorTarget* DragTarget);
+	INTERACTIVETOOLSFRAMEWORK_API virtual void Initialize(IClickBehaviorTarget* ClickTarget, IClickDragBehaviorTarget* DragTarget);
 
 	/**
 	 * Change the Drag BehaviorTarget
 	 */
-	virtual void SetDragTarget(IClickDragBehaviorTarget* DragTarget);
+	INTERACTIVETOOLSFRAMEWORK_API virtual void SetDragTarget(IClickDragBehaviorTarget* DragTarget);
 
 
 	/**
@@ -69,10 +69,10 @@ public:
 
 	// UInputBehavior implementation
 
-	virtual FInputCaptureRequest WantsCapture(const FInputDeviceState& Input) override;
-	virtual FInputCaptureUpdate BeginCapture(const FInputDeviceState& Input, EInputCaptureSide eSide) override;
-	virtual FInputCaptureUpdate UpdateCapture(const FInputDeviceState& Input, const FInputCaptureData& Data) override;
-	virtual void ForceEndCapture(const FInputCaptureData& Data) override;
+	INTERACTIVETOOLSFRAMEWORK_API virtual FInputCaptureRequest WantsCapture(const FInputDeviceState& Input) override;
+	INTERACTIVETOOLSFRAMEWORK_API virtual FInputCaptureUpdate BeginCapture(const FInputDeviceState& Input, EInputCaptureSide eSide) override;
+	INTERACTIVETOOLSFRAMEWORK_API virtual FInputCaptureUpdate UpdateCapture(const FInputDeviceState& Input, const FInputCaptureData& Data) override;
+	INTERACTIVETOOLSFRAMEWORK_API virtual void ForceEndCapture(const FInputCaptureData& Data) override;
 
 
 
@@ -112,22 +112,22 @@ protected:
 	/**
 	 * Internal function that forwards click evens to ClickTarget::OnClicked, you can customize behavior here
 	 */
-	virtual void OnClickedInternal(const FInputDeviceState& Input, const FInputCaptureData& Data);
+	INTERACTIVETOOLSFRAMEWORK_API virtual void OnClickedInternal(const FInputDeviceState& Input, const FInputCaptureData& Data);
 
 	/**
 	* Internal function that forwards click evens to DragTarget::OnClickPress, you can customize behavior here
 	*/
-	virtual void OnClickDragPressInternal(const FInputDeviceState& Input, EInputCaptureSide Side);
+	INTERACTIVETOOLSFRAMEWORK_API virtual void OnClickDragPressInternal(const FInputDeviceState& Input, EInputCaptureSide Side);
 
 	/**
 	* Internal function that forwards click evens to DragTarget::OnClickDrag, you can customize behavior here
 	*/
-	virtual void OnClickDragInternal(const FInputDeviceState& Input, const FInputCaptureData& Data);
+	INTERACTIVETOOLSFRAMEWORK_API virtual void OnClickDragInternal(const FInputDeviceState& Input, const FInputCaptureData& Data);
 
 	/**
 	* Internal function that forwards click evens to DragTarget::OnClickRelease, you can customize behavior here
 	*/
-	virtual void OnClickDragReleaseInternal(const FInputDeviceState& Input, const FInputCaptureData& Data);
+	INTERACTIVETOOLSFRAMEWORK_API virtual void OnClickDragReleaseInternal(const FInputDeviceState& Input, const FInputCaptureData& Data);
 
 
 };
diff --git a/Engine/Source/Runtime/InteractiveToolsFramework/Public/BaseBehaviors/SingleKeyCaptureBehavior.h b/Engine/Source/Runtime/InteractiveToolsFramework/Public/BaseBehaviors/SingleKeyCaptureBehavior.h
--- a/Engine/Source/Runtime/InteractiveToolsFramework/Public/BaseBehaviors/SingleKeyCaptureBehavior.h
+++ b/Engine/Source/Runtime/InteractiveToolsFramework/Public/BaseBehaviors/SingleKeyCaptureBehavior.h
@@ -24,13 +24,13 @@ class UObject;
  * the IModifierToggleBehaviorTarget interface. If you want similar behavior
  * without actually capturing the key, you should use UKeyAsModifierInputBehavior.
  */
-UCLASS()
-class INTERACTIVETOOLSFRAMEWORK_API USingleKeyCaptureBehavior : public UInputBehavior
+UCLASS(MinimalAPI)
+class USingleKeyCaptureBehavior : public UInputBehavior
 {
 	GENERATED_BODY()
 
 public:
-	USingleKeyCaptureBehavior();
+	INTERACTIVETOOLSFRAMEWORK_API USingleKeyCaptureBehavior();
 
 	virtual EInputDevices GetSupportedDevices() override
 	{
@@ -46,7 +46,7 @@ public:
 	 * @param ModifierID integer ID that identifiers the modifier toggle
 	 * @param ModifierKey the key that will be used as the modifier toggle
 	 */
-	virtual void Initialize(IModifierToggleBehaviorTarget* Target, int ModifierID, const FKey& ModifierKey);
+	INTERACTIVETOOLSFRAMEWORK_API virtual void Initialize(IModifierToggleBehaviorTarget* Target, int ModifierID, const FKey& ModifierKey);
 
 	/**
 	 * WantsCapture() will only return capture request if this function returns true (or is null)
@@ -56,10 +56,10 @@ public:
 
 	// UInputBehavior implementation
 
-	virtual FInputCaptureRequest WantsCapture(const FInputDeviceState& Input) override;
-	virtual FInputCaptureUpdate BeginCapture(const FInputDeviceState& Input, EInputCaptureSide Side) override;
-	virtual FInputCaptureUpdate UpdateCapture(const FInputDeviceState& Input, const FInputCaptureData& Data) override;
-	virtual void ForceEndCapture(const FInputCaptureData& Data) override;
+	INTERACTIVETOOLSFRAMEWORK_API virtual FInputCaptureRequest WantsCapture(const FInputDeviceState& Input) override;
+	INTERACTIVETOOLSFRAMEWORK_API virtual FInputCaptureUpdate BeginCapture(const FInputDeviceState& Input, EInputCaptureSide Side) override;
+	INTERACTIVETOOLSFRAMEWORK_API virtual FInputCaptureUpdate UpdateCapture(const FInputDeviceState& Input, const FInputCaptureData& Data) override;
+	INTERACTIVETOOLSFRAMEWORK_API virtual void ForceEndCapture(const FInputCaptureData& Data) override;
 
 
 protected:
diff --git a/Engine/Source/Runtime/InteractiveToolsFramework/Public/BaseBehaviors/Widgets/WidgetBaseBehavior.h b/Engine/Source/Runtime/InteractiveToolsFramework/Public/BaseBehaviors/Widgets/WidgetBaseBehavior.h
--- a/Engine/Source/Runtime/InteractiveToolsFramework/Public/BaseBehaviors/Widgets/WidgetBaseBehavior.h
+++ b/Engine/Source/Runtime/InteractiveToolsFramework/Public/BaseBehaviors/Widgets/WidgetBaseBehavior.h
@@ -21,22 +21,22 @@ class UWidgetBaseBehavior : public UInterface
 	GENERATED_BODY()
 };
 
-class INTERACTIVETOOLSFRAMEWORK_API IWidgetBaseBehavior
+class IWidgetBaseBehavior
 {
 	GENERATED_BODY()
 
 public:
 	// ~Begin SWidget "Proxy" Interface
-	virtual bool OnKeyChar(const FGeometry& MyGeometry, const FCharacterEvent& InCharacterEvent);
-	virtual bool OnKeyDown(const FGeometry& MyGeometry, const FKeyEvent& InKeyEvent);
-	virtual bool OnKeyUp(const FGeometry& MyGeometry, const FKeyEvent& InKeyEvent);
-	virtual bool OnMouseButtonDown(const FGeometry& MyGeometry, const FPointerEvent& MouseEvent);
-	virtual bool OnMouseButtonUp(const FGeometry& MyGeometry, const FPointerEvent& MouseEvent);
-	virtual bool OnMouseMove(const FGeometry& MyGeometry, const FPointerEvent& MouseEvent);
-	virtual bool OnMouseWheel(const FGeometry& MyGeometry, const FPointerEvent& MouseEvent);
-	virtual bool OnMouseButtonDoubleClick(const FGeometry& InMyGeometry, const FPointerEvent& InMouseEvent);
-	virtual bool OnDragDetected(const FGeometry& MyGeometry, const FPointerEvent& MouseEvent);
-	virtual void OnMouseCaptureLost(const FCaptureLostEvent& CaptureLostEvent);
-	virtual int32 OnPaint(const FPaintArgs& Args, const FGeometry& AllottedGeometry, const FSlateRect& MyCullingRect, FSlateWindowElementList& OutDrawElements, int32 LayerId, const FWidgetStyle& InWidgetStyle, bool bParentEnabled) const;
+	INTERACTIVETOOLSFRAMEWORK_API virtual bool OnKeyChar(const FGeometry& MyGeometry, const FCharacterEvent& InCharacterEvent);
+	INTERACTIVETOOLSFRAMEWORK_API virtual bool OnKeyDown(const FGeometry& MyGeometry, const FKeyEvent& InKeyEvent);
+	INTERACTIVETOOLSFRAMEWORK_API virtual bool OnKeyUp(const FGeometry& MyGeometry, const FKeyEvent& InKeyEvent);
+	INTERACTIVETOOLSFRAMEWORK_API virtual bool OnMouseButtonDown(const FGeometry& MyGeometry, const FPointerEvent& MouseEvent);
+	INTERACTIVETOOLSFRAMEWORK_API virtual bool OnMouseButtonUp(const FGeometry& MyGeometry, const FPointerEvent& MouseEvent);
+	INTERACTIVETOOLSFRAMEWORK_API virtual bool OnMouseMove(const FGeometry& MyGeometry, const FPointerEvent& MouseEvent);
+	INTERACTIVETOOLSFRAMEWORK_API virtual bool OnMouseWheel(const FGeometry& MyGeometry, const FPointerEvent& MouseEvent);
+	INTERACTIVETOOLSFRAMEWORK_API virtual bool OnMouseButtonDoubleClick(const FGeometry& InMyGeometry, const FPointerEvent& InMouseEvent);
+	INTERACTIVETOOLSFRAMEWORK_API virtual bool OnDragDetected(const FGeometry& MyGeometry, const FPointerEvent& MouseEvent);
+	INTERACTIVETOOLSFRAMEWORK_API virtual void OnMouseCaptureLost(const FCaptureLostEvent& CaptureLostEvent);
+	INTERACTIVETOOLSFRAMEWORK_API virtual int32 OnPaint(const FPaintArgs& Args, const FGeometry& AllottedGeometry, const FSlateRect& MyCullingRect, FSlateWindowElementList& OutDrawElements, int32 LayerId, const FWidgetStyle& InWidgetStyle, bool bParentEnabled) const;
 	// ~End SWidget "Proxy" Interface
-};
\ No newline at end of file
+};
diff --git a/Engine/Source/Runtime/InteractiveToolsFramework/Public/BaseGizmos/AxisAngleGizmo.h b/Engine/Source/Runtime/InteractiveToolsFramework/Public/BaseGizmos/AxisAngleGizmo.h
--- a/Engine/Source/Runtime/InteractiveToolsFramework/Public/BaseGizmos/AxisAngleGizmo.h
+++ b/Engine/Source/Runtime/InteractiveToolsFramework/Public/BaseGizmos/AxisAngleGizmo.h
@@ -28,45 +28,45 @@ class UObject;
 struct FToolBuilderState;
 
 
-UCLASS()
-class INTERACTIVETOOLSFRAMEWORK_API UAxisAngleGizmoBuilder : public UInteractiveGizmoBuilder
+UCLASS(MinimalAPI)
+class UAxisAngleGizmoBuilder : public UInteractiveGizmoBuilder
 {
 	GENERATED_BODY()
 
 public:
-	virtual UInteractiveGizmo* BuildGizmo(const FToolBuilderState& SceneState) const override;
+	INTERACTIVETOOLSFRAMEWORK_API virtual UInteractiveGizmo* BuildGizmo(const FToolBuilderState& SceneState) const override;
 };
 
 
 /**
  *
  */
-UCLASS()
-class INTERACTIVETOOLSFRAMEWORK_API UAxisAngleGizmo : public UInteractiveGizmo, public IClickDragBehaviorTarget, public IHoverBehaviorTarget
+UCLASS(MinimalAPI)
+class UAxisAngleGizmo : public UInteractiveGizmo, public IClickDragBehaviorTarget, public IHoverBehaviorTarget
 {
 	GENERATED_BODY()
 
 public:
 	// UInteractiveGizmo overrides
 
-	virtual void Setup() override;
+	INTERACTIVETOOLSFRAMEWORK_API virtual void Setup() override;
 
 	// IClickDragBehaviorTarget implementation
 
-	virtual FInputRayHit CanBeginClickDragSequence(const FInputDeviceRay& PressPos) override;
-	virtual void OnClickPress(const FInputDeviceRay& PressPos) override;
-	virtual void OnClickDrag(const FInputDeviceRay& DragPos) override;
-	virtual void OnClickRelease(const FInputDeviceRay& ReleasePos) override;
-	virtual void OnTerminateDragSequence() override;
+	INTERACTIVETOOLSFRAMEWORK_API virtual FInputRayHit CanBeginClickDragSequence(const FInputDeviceRay& PressPos) override;
+	INTERACTIVETOOLSFRAMEWORK_API virtual void OnClickPress(const FInputDeviceRay& PressPos) override;
+	INTERACTIVETOOLSFRAMEWORK_API virtual void OnClickDrag(const FInputDeviceRay& DragPos) override;
+	INTERACTIVETOOLSFRAMEWORK_API virtual void OnClickRelease(const FInputDeviceRay& ReleasePos) override;
+	INTERACTIVETOOLSFRAMEWORK_API virtual void OnTerminateDragSequence() override;
 
 	// IHoverBehaviorTarget implementation
-	virtual FInputRayHit BeginHoverSequenceHitTest(const FInputDeviceRay& PressPos) override;
-	virtual void OnBeginHover(const FInputDeviceRay& DevicePos) override;
-	virtual bool OnUpdateHover(const FInputDeviceRay& DevicePos) override;
-	virtual void OnEndHover() override;
+	INTERACTIVETOOLSFRAMEWORK_API virtual FInputRayHit BeginHoverSequenceHitTest(const FInputDeviceRay& PressPos) override;
+	INTERACTIVETOOLSFRAMEWORK_API virtual void OnBeginHover(const FInputDeviceRay& DevicePos) override;
+	INTERACTIVETOOLSFRAMEWORK_API virtual bool OnUpdateHover(const FInputDeviceRay& DevicePos) override;
+	INTERACTIVETOOLSFRAMEWORK_API virtual void OnEndHover() override;
 
 	// IModifierToggleBehaviorTarget implementation (inherited via IClickDragBehaviorTarget)
-	virtual void OnUpdateModifierState(int ModifierID, bool bIsOn) override;
+	INTERACTIVETOOLSFRAMEWORK_API virtual void OnUpdateModifierState(int ModifierID, bool bIsOn) override;
 
 public:
 	UPROPERTY()
diff --git a/Engine/Source/Runtime/InteractiveToolsFramework/Public/BaseGizmos/AxisPositionGizmo.h b/Engine/Source/Runtime/InteractiveToolsFramework/Public/BaseGizmos/AxisPositionGizmo.h
--- a/Engine/Source/Runtime/InteractiveToolsFramework/Public/BaseGizmos/AxisPositionGizmo.h
+++ b/Engine/Source/Runtime/InteractiveToolsFramework/Public/BaseGizmos/AxisPositionGizmo.h
@@ -28,13 +28,13 @@ class UGizmoViewContext;
 class UObject;
 struct FToolBuilderState;
 
-UCLASS()
-class INTERACTIVETOOLSFRAMEWORK_API UAxisPositionGizmoBuilder : public UInteractiveGizmoBuilder
+UCLASS(MinimalAPI)
+class UAxisPositionGizmoBuilder : public UInteractiveGizmoBuilder
 {
 	GENERATED_BODY()
 
 public:
-	virtual UInteractiveGizmo* BuildGizmo(const FToolBuilderState& SceneState) const override;
+	INTERACTIVETOOLSFRAMEWORK_API virtual UInteractiveGizmo* BuildGizmo(const FToolBuilderState& SceneState) const override;
 };
 
 
@@ -54,29 +54,29 @@ public:
  *
  * Internally a UClickDragInputBehavior is used to handle mouse input, configured in ::Setup()
  */
-UCLASS()
-class INTERACTIVETOOLSFRAMEWORK_API UAxisPositionGizmo : public UInteractiveGizmo, public IClickDragBehaviorTarget, public IHoverBehaviorTarget
+UCLASS(MinimalAPI)
+class UAxisPositionGizmo : public UInteractiveGizmo, public IClickDragBehaviorTarget, public IHoverBehaviorTarget
 {
 	GENERATED_BODY()
 
 public:
 	// UInteractiveGizmo overrides
 
-	virtual void Setup() override;
+	INTERACTIVETOOLSFRAMEWORK_API virtual void Setup() override;
 
 	// IClickDragBehaviorTarget implementation
 
-	virtual FInputRayHit CanBeginClickDragSequence(const FInputDeviceRay& PressPos) override;
-	virtual void OnClickPress(const FInputDeviceRay& PressPos) override;
-	virtual void OnClickDrag(const FInputDeviceRay& DragPos) override;
-	virtual void OnClickRelease(const FInputDeviceRay& ReleasePos) override;
-	virtual void OnTerminateDragSequence() override;
+	INTERACTIVETOOLSFRAMEWORK_API virtual FInputRayHit CanBeginClickDragSequence(const FInputDeviceRay& PressPos) override;
+	INTERACTIVETOOLSFRAMEWORK_API virtual void OnClickPress(const FInputDeviceRay& PressPos) override;
+	INTERACTIVETOOLSFRAMEWORK_API virtual void OnClickDrag(const FInputDeviceRay& DragPos) override;
+	INTERACTIVETOOLSFRAMEWORK_API virtual void OnClickRelease(const FInputDeviceRay& ReleasePos) override;
+	INTERACTIVETOOLSFRAMEWORK_API virtual void OnTerminateDragSequence() override;
 
 	// IHoverBehaviorTarget implementation
-	virtual FInputRayHit BeginHoverSequenceHitTest(const FInputDeviceRay& PressPos) override;
-	virtual void OnBeginHover(const FInputDeviceRay& DevicePos) override;
-	virtual bool OnUpdateHover(const FInputDeviceRay& DevicePos) override;
-	virtual void OnEndHover() override;
+	INTERACTIVETOOLSFRAMEWORK_API virtual FInputRayHit BeginHoverSequenceHitTest(const FInputDeviceRay& PressPos) override;
+	INTERACTIVETOOLSFRAMEWORK_API virtual void OnBeginHover(const FInputDeviceRay& DevicePos) override;
+	INTERACTIVETOOLSFRAMEWORK_API virtual bool OnUpdateHover(const FInputDeviceRay& DevicePos) override;
+	INTERACTIVETOOLSFRAMEWORK_API virtual void OnEndHover() override;
 
 
 public:
diff --git a/Engine/Source/Runtime/InteractiveToolsFramework/Public/BaseGizmos/AxisSources.h b/Engine/Source/Runtime/InteractiveToolsFramework/Public/BaseGizmos/AxisSources.h
--- a/Engine/Source/Runtime/InteractiveToolsFramework/Public/BaseGizmos/AxisSources.h
+++ b/Engine/Source/Runtime/InteractiveToolsFramework/Public/BaseGizmos/AxisSources.h
@@ -21,8 +21,8 @@ class USceneComponent;
  * UGizmoConstantAxisSource is an IGizmoAxisSource implementation that
  * internally stores the Origin and Direction constants
  */
-UCLASS()
-class INTERACTIVETOOLSFRAMEWORK_API UGizmoConstantAxisSource : public UObject, public IGizmoAxisSource
+UCLASS(MinimalAPI)
+class UGizmoConstantAxisSource : public UObject, public IGizmoAxisSource
 {
 	GENERATED_BODY()
 public:
@@ -51,8 +51,8 @@ public:
  * UGizmoConstantFrameAxisSource is an IGizmoAxisSource implementation that
  * internally stores the Origin, Direction, and X/Y Tangent constants.
  */
-UCLASS()
-class INTERACTIVETOOLSFRAMEWORK_API UGizmoConstantFrameAxisSource : public UObject, public IGizmoAxisSource
+UCLASS(MinimalAPI)
+class UGizmoConstantFrameAxisSource : public UObject, public IGizmoAxisSource
 {
 	GENERATED_BODY()
 public:
@@ -94,8 +94,8 @@ public:
  * (ie X/Y/Z axes) based on an integer AxisIndex in range [0,2]. The Orgin is
  * internally stored as a FProperty.
  */
-UCLASS()
-class INTERACTIVETOOLSFRAMEWORK_API UGizmoWorldAxisSource : public UObject, public IGizmoAxisSource
+UCLASS(MinimalAPI)
+class UGizmoWorldAxisSource : public UObject, public IGizmoAxisSource
 {
 	GENERATED_BODY()
 public:
@@ -139,18 +139,18 @@ public:
  * X/Y/Z axes of a Component's local coordinate system, mapped to World, based on an integer AxisIndex in range [0,2].
  * The Axis Origin is the Component's transform origin. Tangent vectors are provided.
  */
-UCLASS()
-class INTERACTIVETOOLSFRAMEWORK_API UGizmoComponentAxisSource : public UObject, public IGizmoAxisSource
+UCLASS(MinimalAPI)
+class UGizmoComponentAxisSource : public UObject, public IGizmoAxisSource
 {
 	GENERATED_BODY()
 public:
-	virtual FVector GetOrigin() const;
+	INTERACTIVETOOLSFRAMEWORK_API virtual FVector GetOrigin() const;
 
-	virtual FVector GetDirection() const;
+	INTERACTIVETOOLSFRAMEWORK_API virtual FVector GetDirection() const;
 
 	virtual bool HasTangentVectors() const { return true; }
 
-	virtual void GetTangentVectors(FVector& TangentXOut, FVector& TangentYOut) const;
+	INTERACTIVETOOLSFRAMEWORK_API virtual void GetTangentVectors(FVector& TangentXOut, FVector& TangentYOut) const;
 
 	UPROPERTY()
 	TObjectPtr<USceneComponent> Component;
@@ -176,4 +176,4 @@ public:
 		NewSource->bLocalAxes = bUseLocalAxes;
 		return NewSource;
 	}
-};
\ No newline at end of file
+};
diff --git a/Engine/Source/Runtime/InteractiveToolsFramework/Public/BaseGizmos/BrushStampIndicator.h b/Engine/Source/Runtime/InteractiveToolsFramework/Public/BaseGizmos/BrushStampIndicator.h
--- a/Engine/Source/Runtime/InteractiveToolsFramework/Public/BaseGizmos/BrushStampIndicator.h
+++ b/Engine/Source/Runtime/InteractiveToolsFramework/Public/BaseGizmos/BrushStampIndicator.h
@@ -9,12 +9,12 @@
 class UPrimitiveComponent;
 class UPreviewMesh;
 
-UCLASS()
-class INTERACTIVETOOLSFRAMEWORK_API UBrushStampIndicatorBuilder : public UInteractiveGizmoBuilder
+UCLASS(MinimalAPI)
+class UBrushStampIndicatorBuilder : public UInteractiveGizmoBuilder
 {
 	GENERATED_BODY()
 public:
-	virtual UInteractiveGizmo* BuildGizmo(const FToolBuilderState& SceneState) const override;
+	INTERACTIVETOOLSFRAMEWORK_API virtual UInteractiveGizmo* BuildGizmo(const FToolBuilderState& SceneState) const override;
 };
 
 
@@ -22,8 +22,8 @@ public:
 /*
  * UBrushStampIndicator is a simple 3D brush indicator. 
  */
-UCLASS(Transient)
-class INTERACTIVETOOLSFRAMEWORK_API UBrushStampIndicator : public UInteractiveGizmo
+UCLASS(Transient, MinimalAPI)
+class UBrushStampIndicator : public UInteractiveGizmo
 {
 	GENERATED_BODY()
 
@@ -31,21 +31,21 @@ public:
 
 	// UInteractiveGizmo interface/implementation
 
-	virtual void Setup() override;
-	virtual void Shutdown() override;
-	virtual void Render(IToolsContextRenderAPI* RenderAPI) override;
-	virtual void Tick(float DeltaTime) override;
+	INTERACTIVETOOLSFRAMEWORK_API virtual void Setup() override;
+	INTERACTIVETOOLSFRAMEWORK_API virtual void Shutdown() override;
+	INTERACTIVETOOLSFRAMEWORK_API virtual void Render(IToolsContextRenderAPI* RenderAPI) override;
+	INTERACTIVETOOLSFRAMEWORK_API virtual void Tick(float DeltaTime) override;
 
 
 	/**
 	 * Update the Radius, Position, and Normal of the stamp indicator
 	 */
-	virtual void Update(float Radius, const FVector& Position, const FVector& Normal, float Falloff);
+	INTERACTIVETOOLSFRAMEWORK_API virtual void Update(float Radius, const FVector& Position, const FVector& Normal, float Falloff);
 
 	/**
 	* Update the Radius, Transform and Falloff of the Stamp Indicator
 	*/
-	virtual void Update(float Radius, const FTransform& WorldTransform, float Falloff);
+	INTERACTIVETOOLSFRAMEWORK_API virtual void Update(float Radius, const FTransform& WorldTransform, float Falloff);
 
 public:
 
@@ -106,4 +106,4 @@ public:
 protected:
 	UPrimitiveComponent* ScaleInitializedComponent = nullptr;		// we are just using this as a key, never calling any functions on it
 	FVector InitialComponentScale;
-};
\ No newline at end of file
+};
diff --git a/Engine/Source/Runtime/InteractiveToolsFramework/Public/BaseGizmos/CombinedTransformGizmo.h b/Engine/Source/Runtime/InteractiveToolsFramework/Public/BaseGizmos/CombinedTransformGizmo.h
--- a/Engine/Source/Runtime/InteractiveToolsFramework/Public/BaseGizmos/CombinedTransformGizmo.h
+++ b/Engine/Source/Runtime/InteractiveToolsFramework/Public/BaseGizmos/CombinedTransformGizmo.h
@@ -35,13 +35,13 @@ class FTransformGizmoTransformChange;
  * The static factory method ::ConstructDefault3AxisGizmo() creates and initializes an 
  * Actor suitable for use in a standard 3-axis Transformation Gizmo.
  */
-UCLASS(Transient, NotPlaceable, Hidden, NotBlueprintable, NotBlueprintType)
-class INTERACTIVETOOLSFRAMEWORK_API ACombinedTransformGizmoActor : public AGizmoActor
+UCLASS(Transient, NotPlaceable, Hidden, NotBlueprintable, NotBlueprintType, MinimalAPI)
+class ACombinedTransformGizmoActor : public AGizmoActor
 {
 	GENERATED_BODY()
 public:
 
-	ACombinedTransformGizmoActor();
+	INTERACTIVETOOLSFRAMEWORK_API ACombinedTransformGizmoActor();
 
 public:
 	//
@@ -134,7 +134,7 @@ public:
 	 * Create a new instance of ACombinedTransformGizmoActor and populate the various
 	 * sub-components with standard GizmoXComponent instances suitable for a 3-axis transformer Gizmo
 	 */
-	static ACombinedTransformGizmoActor* ConstructDefault3AxisGizmo(
+	static INTERACTIVETOOLSFRAMEWORK_API ACombinedTransformGizmoActor* ConstructDefault3AxisGizmo(
 		UWorld* World, UGizmoViewContext* GizmoViewContext
 	);
 
@@ -142,7 +142,7 @@ public:
 	 * Create a new instance of ACombinedTransformGizmoActor. Populate the sub-components 
 	 * specified by Elements with standard GizmoXComponent instances suitable for a 3-axis transformer Gizmo
 	 */
-	static ACombinedTransformGizmoActor* ConstructCustom3AxisGizmo(
+	static INTERACTIVETOOLSFRAMEWORK_API ACombinedTransformGizmoActor* ConstructCustom3AxisGizmo(
 		UWorld* World, UGizmoViewContext* GizmoViewContext,
 		ETransformGizmoSubElements Elements
 	);
@@ -163,7 +163,7 @@ public:
  * a three-axis transformation Gizmo, override this function to customize
  * the Actor sub-elements.
  */
-class INTERACTIVETOOLSFRAMEWORK_API FCombinedTransformGizmoActorFactory
+class FCombinedTransformGizmoActorFactory
 {
 public:
 	FCombinedTransformGizmoActorFactory(UGizmoViewContext* GizmoViewContextIn)
@@ -184,7 +184,7 @@ public:
 	 * @param World the UWorld to create the new Actor in
 	 * @return new ACombinedTransformGizmoActor instance with members initialized with Components suitable for a transformation Gizmo
 	 */
-	virtual ACombinedTransformGizmoActor* CreateNewGizmoActor(UWorld* World) const;
+	INTERACTIVETOOLSFRAMEWORK_API virtual ACombinedTransformGizmoActor* CreateNewGizmoActor(UWorld* World) const;
 
 protected:
 	/**
@@ -200,8 +200,8 @@ protected:
 
 
 
-UCLASS()
-class INTERACTIVETOOLSFRAMEWORK_API UCombinedTransformGizmoBuilder : public UInteractiveGizmoBuilder
+UCLASS(MinimalAPI)
+class UCombinedTransformGizmoBuilder : public UInteractiveGizmoBuilder
 {
 	GENERATED_BODY()
 
@@ -242,7 +242,7 @@ public:
 	TFunction<void(UPrimitiveComponent*, EToolContextCoordinateSystem)> UpdateCoordSystemFunction;
 
 
-	virtual UInteractiveGizmo* BuildGizmo(const FToolBuilderState& SceneState) const override;
+	INTERACTIVETOOLSFRAMEWORK_API virtual UInteractiveGizmo* BuildGizmo(const FToolBuilderState& SceneState) const override;
 };
 
 
@@ -252,7 +252,7 @@ public:
  * This struct does not directly do anything, it just wraps up the multiple flags/states
  * needed to provide such functionality
  */
-struct INTERACTIVETOOLSFRAMEWORK_API FToolContextOptionalToggle
+struct FToolContextOptionalToggle
 {
 	bool bEnabledDirectly = false;
 	bool bEnabledInContext = false;
@@ -288,18 +288,18 @@ struct INTERACTIVETOOLSFRAMEWORK_API FToolContextOptionalToggle
  * a suitably-configured GizmoActor and everything else will be handled automatically.
  * 
  */
-UCLASS()
-class INTERACTIVETOOLSFRAMEWORK_API UCombinedTransformGizmo : public UInteractiveGizmo
+UCLASS(MinimalAPI)
+class UCombinedTransformGizmo : public UInteractiveGizmo
 {
 	GENERATED_BODY()
 
 public:
 
-	virtual void SetWorld(UWorld* World);
-	virtual void SetGizmoActorBuilder(TSharedPtr<FCombinedTransformGizmoActorFactory> Builder);
-	virtual void SetSubGizmoBuilderIdentifiers(FString AxisPositionBuilderIdentifier, FString PlanePositionBuilderIdentifier, FString AxisAngleBuilderIdentifier);
-	virtual void SetUpdateHoverFunction(TFunction<void(UPrimitiveComponent*, bool)> HoverFunction);
-	virtual void SetUpdateCoordSystemFunction(TFunction<void(UPrimitiveComponent*, EToolContextCoordinateSystem)> CoordSysFunction);
+	INTERACTIVETOOLSFRAMEWORK_API virtual void SetWorld(UWorld* World);
+	INTERACTIVETOOLSFRAMEWORK_API virtual void SetGizmoActorBuilder(TSharedPtr<FCombinedTransformGizmoActorFactory> Builder);
+	INTERACTIVETOOLSFRAMEWORK_API virtual void SetSubGizmoBuilderIdentifiers(FString AxisPositionBuilderIdentifier, FString PlanePositionBuilderIdentifier, FString AxisAngleBuilderIdentifier);
+	INTERACTIVETOOLSFRAMEWORK_API virtual void SetUpdateHoverFunction(TFunction<void(UPrimitiveComponent*, bool)> HoverFunction);
+	INTERACTIVETOOLSFRAMEWORK_API virtual void SetUpdateCoordSystemFunction(TFunction<void(UPrimitiveComponent*, EToolContextCoordinateSystem)> CoordSysFunction);
 	
 	/**
 	 * If used, binds alignment functions to the sub gizmos that they can use to align to geometry in the scene. 
@@ -309,7 +309,7 @@ public:
 	 * Subgizmos align to the point in different ways, usually by projecting onto the axis or plane that they
 	 * operate in.
 	 */
-	virtual void SetWorldAlignmentFunctions(
+	INTERACTIVETOOLSFRAMEWORK_API virtual void SetWorldAlignmentFunctions(
 		TUniqueFunction<bool()>&& ShouldAlignDestination,
 		TUniqueFunction<bool(const FRay&, FVector&)>&& DestinationAlignmentRayCaster
 		);
@@ -320,7 +320,7 @@ public:
 	 * locally scaled. However, this can be changed to a custom check here, perhaps to hide them in extra
 	 * conditions or to always show them (if the gizmo is not scaling a component).
 	 */
-	virtual void SetIsNonUniformScaleAllowedFunction(
+	INTERACTIVETOOLSFRAMEWORK_API virtual void SetIsNonUniformScaleAllowedFunction(
 		TUniqueFunction<bool()>&& IsNonUniformScaleAllowed
 	);
 
@@ -336,24 +336,24 @@ public:
 	 *
 	 * TODO: Should this affect uniform scaling too?
 	 */
-	virtual void SetDisallowNegativeScaling(bool bDisallow);
+	INTERACTIVETOOLSFRAMEWORK_API virtual void SetDisallowNegativeScaling(bool bDisallow);
 
 	// UInteractiveGizmo overrides
-	virtual void Setup() override;
-	virtual void Shutdown() override;
-	virtual void Tick(float DeltaTime) override;
+	INTERACTIVETOOLSFRAMEWORK_API virtual void Setup() override;
+	INTERACTIVETOOLSFRAMEWORK_API virtual void Shutdown() override;
+	INTERACTIVETOOLSFRAMEWORK_API virtual void Tick(float DeltaTime) override;
 
 	/**
 	 * Set the active target object for the Gizmo
 	 * @param Target active target
 	 * @param TransactionProvider optional IToolContextTransactionProvider implementation to use - by default uses GizmoManager
 	 */
-	virtual void SetActiveTarget(UTransformProxy* Target, IToolContextTransactionProvider* TransactionProvider = nullptr);
+	INTERACTIVETOOLSFRAMEWORK_API virtual void SetActiveTarget(UTransformProxy* Target, IToolContextTransactionProvider* TransactionProvider = nullptr);
 
 	/**
 	 * Clear the active target object for the Gizmo
 	 */
-	virtual void ClearActiveTarget();
+	INTERACTIVETOOLSFRAMEWORK_API virtual void ClearActiveTarget();
 
 	/** The active target object for the Gizmo */
 	UPROPERTY()
@@ -367,7 +367,7 @@ public:
 	/**
 	 * @return current transform of Gizmo
 	 */
-	FTransform GetGizmoTransform() const;
+	INTERACTIVETOOLSFRAMEWORK_API FTransform GetGizmoTransform() const;
 
 	/**
 	 * Repositions the gizmo without issuing undo/redo changes, triggering callbacks, 
@@ -375,41 +375,41 @@ public:
 	 * it being viewed as a gizmo manipulation.
 	 * @param bKeepGizmoUnscaled If true, the scale component of NewTransform is passed through to the target but gizmo scale is set to 1
 	 */
-	void ReinitializeGizmoTransform(const FTransform& NewTransform, bool bKeepGizmoUnscaled = true);
+	INTERACTIVETOOLSFRAMEWORK_API void ReinitializeGizmoTransform(const FTransform& NewTransform, bool bKeepGizmoUnscaled = true);
 
 	/**
 	 * Set a new position for the Gizmo. This is done via the same mechanisms as the sub-gizmos,
 	 * so it generates the same Change/Modify() events, and hence works with Undo/Redo
 	 * @param bKeepGizmoUnscaled If true, the scale component of NewTransform is passed through to the target but gizmo scale is set to 1
 	 */
-	virtual void SetNewGizmoTransform(const FTransform& NewTransform, bool bKeepGizmoUnscaled = true);
+	INTERACTIVETOOLSFRAMEWORK_API virtual void SetNewGizmoTransform(const FTransform& NewTransform, bool bKeepGizmoUnscaled = true);
 
 	/**
 	 * Called at the start of a sequence of gizmo transform edits, for instance while dragging or
 	 * manipulating the gizmo numerical UI.
 	 */
-	virtual void BeginTransformEditSequence();
+	INTERACTIVETOOLSFRAMEWORK_API virtual void BeginTransformEditSequence();
 
 	/**
 	 * Called at the end of a sequence of gizmo transform edits.
 	 */
-	virtual void EndTransformEditSequence();
+	INTERACTIVETOOLSFRAMEWORK_API virtual void EndTransformEditSequence();
 
 	/**
 	 * Updates the gizmo transform between Begin/EndTransformeditSequence calls.
 	 */
-	void UpdateTransformDuringEditSequence(const FTransform& NewTransform, bool bKeepGizmoUnscaled = true);
+	INTERACTIVETOOLSFRAMEWORK_API void UpdateTransformDuringEditSequence(const FTransform& NewTransform, bool bKeepGizmoUnscaled = true);
 
 	/**
 	 * Explicitly set the child scale. Mainly useful to "reset" the child scale to (1,1,1) when re-using Gizmo across multiple transform actions.
 	 * @warning does not generate change/modify events!!
 	 */
-	virtual void SetNewChildScale(const FVector& NewChildScale);
+	INTERACTIVETOOLSFRAMEWORK_API virtual void SetNewChildScale(const FVector& NewChildScale);
 
 	/**
 	 * Set visibility for this Gizmo
 	 */
-	virtual void SetVisibility(bool bVisible);
+	INTERACTIVETOOLSFRAMEWORK_API virtual void SetVisibility(bool bVisible);
 
 	/**
 	 * @return true if Gizmo is visible
@@ -489,7 +489,7 @@ public:
 	 * Gets the elements that this gizmo was initialized with. Note that this may not account for individual
 	 * element visibility- for instance the scaling component may not be visible if IsNonUniformScaleAllowed() is false.
 	 */
-	ETransformGizmoSubElements GetGizmoElements();
+	INTERACTIVETOOLSFRAMEWORK_API ETransformGizmoSubElements GetGizmoElements();
 
 
 	/**
@@ -503,7 +503,7 @@ public:
 	 * Note that this is an optional feature, as it would require tools to keep this transform up to date if they
 	 * want the UI to use it, so tools could just leave it unset.
 	 */
-	void SetDisplaySpaceTransform(TOptional<FTransform> TransformIn);
+	INTERACTIVETOOLSFRAMEWORK_API void SetDisplaySpaceTransform(TOptional<FTransform> TransformIn);
 	const TOptional<FTransform>& GetDisplaySpaceTransform() { return DisplaySpaceTransform; }
 
 	/**
@@ -587,7 +587,7 @@ protected:
 	TObjectPtr<UGizmoConstantFrameAxisSource> CameraAxisSource;
 
 	// internal function that updates CameraAxisSource by getting current view state from GizmoManager
-	void UpdateCameraAxisSource();
+	INTERACTIVETOOLSFRAMEWORK_API void UpdateCameraAxisSource();
 
 
 	/** X-axis source is shared across Gizmos, and created internally during SetActiveTarget() */
@@ -650,7 +650,7 @@ protected:
 
 
 	/** @return a new instance of the standard axis-translation Gizmo */
-	virtual UInteractiveGizmo* AddAxisTranslationGizmo(
+	INTERACTIVETOOLSFRAMEWORK_API virtual UInteractiveGizmo* AddAxisTranslationGizmo(
 		UPrimitiveComponent* AxisComponent, USceneComponent* RootComponent,
 		IGizmoAxisSource* AxisSource,
 		IGizmoTransformSource* TransformSource, 
@@ -658,7 +658,7 @@ protected:
 		int AxisIndex);
 
 	/** @return a new instance of the standard plane-translation Gizmo */
-	virtual UInteractiveGizmo* AddPlaneTranslationGizmo(
+	INTERACTIVETOOLSFRAMEWORK_API virtual UInteractiveGizmo* AddPlaneTranslationGizmo(
 		UPrimitiveComponent* AxisComponent, USceneComponent* RootComponent,
 		IGizmoAxisSource* AxisSource,
 		IGizmoTransformSource* TransformSource,
@@ -666,37 +666,37 @@ protected:
 		int XAxisIndex, int YAxisIndex);
 
 	/** @return a new instance of the standard axis-rotation Gizmo */
-	virtual UInteractiveGizmo* AddAxisRotationGizmo(
+	INTERACTIVETOOLSFRAMEWORK_API virtual UInteractiveGizmo* AddAxisRotationGizmo(
 		UPrimitiveComponent* AxisComponent, USceneComponent* RootComponent,
 		IGizmoAxisSource* AxisSource,
 		IGizmoTransformSource* TransformSource,
 		IGizmoStateTarget* StateTarget);
 
 	/** @return a new instance of the standard axis-scaling Gizmo */
-	virtual UInteractiveGizmo* AddAxisScaleGizmo(
+	INTERACTIVETOOLSFRAMEWORK_API virtual UInteractiveGizmo* AddAxisScaleGizmo(
 		UPrimitiveComponent* AxisComponent, USceneComponent* RootComponent,
 		IGizmoAxisSource* GizmoAxisSource, IGizmoAxisSource* ParameterAxisSource,
 		IGizmoTransformSource* TransformSource,
 		IGizmoStateTarget* StateTarget);
 
 	/** @return a new instance of the standard plane-scaling Gizmo */
-	virtual UInteractiveGizmo* AddPlaneScaleGizmo(
+	INTERACTIVETOOLSFRAMEWORK_API virtual UInteractiveGizmo* AddPlaneScaleGizmo(
 		UPrimitiveComponent* AxisComponent, USceneComponent* RootComponent,
 		IGizmoAxisSource* GizmoAxisSource, IGizmoAxisSource* ParameterAxisSource,
 		IGizmoTransformSource* TransformSource,
 		IGizmoStateTarget* StateTarget);
 
 	/** @return a new instance of the standard plane-scaling Gizmo */
-	virtual UInteractiveGizmo* AddUniformScaleGizmo(
+	INTERACTIVETOOLSFRAMEWORK_API virtual UInteractiveGizmo* AddUniformScaleGizmo(
 		UPrimitiveComponent* ScaleComponent, USceneComponent* RootComponent,
 		IGizmoAxisSource* GizmoAxisSource, IGizmoAxisSource* ParameterAxisSource,
 		IGizmoTransformSource* TransformSource,
 		IGizmoStateTarget* StateTarget);
 
 	// Axis and Plane TransformSources use these function to execute snapping queries
-	bool PositionSnapFunction(const FVector& WorldPosition, FVector& SnappedPositionOut) const;
-	bool PositionAxisDeltaSnapFunction(double AxisDelta, double& SnappedDeltaOut, int AxisIndex) const;
-	FQuat RotationSnapFunction(const FQuat& DeltaRotation) const;
-	bool RotationAxisAngleSnapFunction(double AxisAngleDelta, double& SnappedAxisAngleDeltaOut, int AxisIndex) const;
+	INTERACTIVETOOLSFRAMEWORK_API bool PositionSnapFunction(const FVector& WorldPosition, FVector& SnappedPositionOut) const;
+	INTERACTIVETOOLSFRAMEWORK_API bool PositionAxisDeltaSnapFunction(double AxisDelta, double& SnappedDeltaOut, int AxisIndex) const;
+	INTERACTIVETOOLSFRAMEWORK_API FQuat RotationSnapFunction(const FQuat& DeltaRotation) const;
+	INTERACTIVETOOLSFRAMEWORK_API bool RotationAxisAngleSnapFunction(double AxisAngleDelta, double& SnappedAxisAngleDeltaOut, int AxisIndex) const;
 
 };
diff --git a/Engine/Source/Runtime/InteractiveToolsFramework/Public/BaseGizmos/GizmoActor.h b/Engine/Source/Runtime/InteractiveToolsFramework/Public/BaseGizmos/GizmoActor.h
--- a/Engine/Source/Runtime/InteractiveToolsFramework/Public/BaseGizmos/GizmoActor.h
+++ b/Engine/Source/Runtime/InteractiveToolsFramework/Public/BaseGizmos/GizmoActor.h
@@ -19,39 +19,39 @@ class UGizmoViewContext;
  * Currently this does not involve any special functionality, but a set of static functions
  * are provided to create default Components commonly used in Gizmos.
  */
-UCLASS(Transient, NotPlaceable, Hidden, NotBlueprintable, NotBlueprintType)
-class INTERACTIVETOOLSFRAMEWORK_API AGizmoActor : public AInternalToolFrameworkActor
+UCLASS(Transient, NotPlaceable, Hidden, NotBlueprintable, NotBlueprintType, MinimalAPI)
+class AGizmoActor : public AInternalToolFrameworkActor
 {
 	GENERATED_BODY()
 public:
-	AGizmoActor();
+	INTERACTIVETOOLSFRAMEWORK_API AGizmoActor();
 
 	/** Add standard arrow component to Actor, generally used for axis-translation */
-	static UGizmoArrowComponent* AddDefaultArrowComponent(
+	static INTERACTIVETOOLSFRAMEWORK_API UGizmoArrowComponent* AddDefaultArrowComponent(
 		UWorld* World, AActor* Actor, UGizmoViewContext* GizmoViewContext,
 		const FLinearColor& Color, const FVector& LocalDirection, const float Length = 80.0f
 	);
 	/** Add standard rectangle component to Actor, generally used for plane-translation */
-	static UGizmoRectangleComponent* AddDefaultRectangleComponent(
+	static INTERACTIVETOOLSFRAMEWORK_API UGizmoRectangleComponent* AddDefaultRectangleComponent(
 		UWorld* World, AActor* Actor, UGizmoViewContext* GizmoViewContext,
 		const FLinearColor& Color, const FVector& PlaneAxis1, const FVector& PlaneAxisx2
 	);
 	/** Add standard circle component to Actor, generally used for axis-rotation */
-	static UGizmoCircleComponent* AddDefaultCircleComponent(
+	static INTERACTIVETOOLSFRAMEWORK_API UGizmoCircleComponent* AddDefaultCircleComponent(
 		UWorld* World, AActor* Actor, UGizmoViewContext* GizmoViewContext,
 		const FLinearColor& Color, const FVector& PlaneNormal,
 		float Radius = 120.0f
 	);
 
 	/** Add standard 3D box component to Actor. By default the box is axis-aligned, centered at the specified Origin */
-	static UGizmoBoxComponent* AddDefaultBoxComponent(
+	static INTERACTIVETOOLSFRAMEWORK_API UGizmoBoxComponent* AddDefaultBoxComponent(
 		UWorld* World, AActor* Actor, UGizmoViewContext* GizmoViewContext,
 		const FLinearColor& Color, const FVector& Origin, 
 		const FVector& Dimensions = FVector(20.0f, 20.0f, 20.0f)
 	);
 
 	/** Add standard disk component to Actor, generally used for handles */
-	static UGizmoLineHandleComponent* AddDefaultLineHandleComponent(
+	static INTERACTIVETOOLSFRAMEWORK_API UGizmoLineHandleComponent* AddDefaultLineHandleComponent(
 		UWorld* World, AActor* Actor, UGizmoViewContext* GizmoViewContext,
 		const FLinearColor& Color, const FVector& PlaneNormal, const FVector& LocalDirection,
 		const float Length = 60.f, const bool bImageScale = false
diff --git a/Engine/Source/Runtime/InteractiveToolsFramework/Public/BaseGizmos/GizmoArrowComponent.h b/Engine/Source/Runtime/InteractiveToolsFramework/Public/BaseGizmos/GizmoArrowComponent.h
--- a/Engine/Source/Runtime/InteractiveToolsFramework/Public/BaseGizmos/GizmoArrowComponent.h
+++ b/Engine/Source/Runtime/InteractiveToolsFramework/Public/BaseGizmos/GizmoArrowComponent.h
@@ -11,8 +11,8 @@
  * Simple Component intended to be used as part of 3D Gizmos.
  * Currently draws the "arrow" as a single 3D line.
  */
-UCLASS(ClassGroup = Utility, HideCategories = (Physics, Collision, Lighting, Rendering, Mobile))
-class INTERACTIVETOOLSFRAMEWORK_API UGizmoArrowComponent : public UGizmoBaseComponent
+UCLASS(ClassGroup = Utility, HideCategories = (Physics, Collision, Lighting, Rendering, Mobile), MinimalAPI)
+class UGizmoArrowComponent : public UGizmoBaseComponent
 {
 	GENERATED_BODY()
 
@@ -31,11 +31,11 @@ public:
 
 private:
 	//~ Begin UPrimitiveComponent Interface.
-	virtual FPrimitiveSceneProxy* CreateSceneProxy() override;
-	virtual bool LineTraceComponent(FHitResult& OutHit, const FVector Start, const FVector End, const FCollisionQueryParams& Params) override;
+	INTERACTIVETOOLSFRAMEWORK_API virtual FPrimitiveSceneProxy* CreateSceneProxy() override;
+	INTERACTIVETOOLSFRAMEWORK_API virtual bool LineTraceComponent(FHitResult& OutHit, const FVector Start, const FVector End, const FCollisionQueryParams& Params) override;
 	//~ End UPrimitiveComponent Interface.
 
 	//~ Begin USceneComponent Interface.
-	virtual FBoxSphereBounds CalcBounds(const FTransform& LocalToWorld) const override;
+	INTERACTIVETOOLSFRAMEWORK_API virtual FBoxSphereBounds CalcBounds(const FTransform& LocalToWorld) const override;
 	//~ Begin USceneComponent Interface.
-};
\ No newline at end of file
+};
diff --git a/Engine/Source/Runtime/InteractiveToolsFramework/Public/BaseGizmos/GizmoBaseComponent.h b/Engine/Source/Runtime/InteractiveToolsFramework/Public/BaseGizmos/GizmoBaseComponent.h
--- a/Engine/Source/Runtime/InteractiveToolsFramework/Public/BaseGizmos/GizmoBaseComponent.h
+++ b/Engine/Source/Runtime/InteractiveToolsFramework/Public/BaseGizmos/GizmoBaseComponent.h
@@ -14,8 +14,8 @@ class UGizmoViewContext;
  * Contains common properties and utility functions.
  * This class does nothing by itself, use subclasses like UGizmoCircleComponent
  */
-UCLASS(ClassGroup = Utility, HideCategories = (Physics, Collision, Mobile))
-class INTERACTIVETOOLSFRAMEWORK_API UGizmoBaseComponent : public UPrimitiveComponent
+UCLASS(ClassGroup = Utility, HideCategories = (Physics, Collision, Mobile), MinimalAPI)
+class UGizmoBaseComponent : public UPrimitiveComponent
 {
 	GENERATED_BODY()
 
@@ -87,4 +87,4 @@ protected:
 	// True when GizmoViewContext is not null. Here as a boolean so it
 	// can be pointed to by the proxy to see what it should do.
 	bool bIsViewDependent = false;
-};
\ No newline at end of file
+};
diff --git a/Engine/Source/Runtime/InteractiveToolsFramework/Public/BaseGizmos/GizmoBoxComponent.h b/Engine/Source/Runtime/InteractiveToolsFramework/Public/BaseGizmos/GizmoBoxComponent.h
--- a/Engine/Source/Runtime/InteractiveToolsFramework/Public/BaseGizmos/GizmoBoxComponent.h
+++ b/Engine/Source/Runtime/InteractiveToolsFramework/Public/BaseGizmos/GizmoBoxComponent.h
@@ -11,8 +11,8 @@
  * Simple Component intended to be used as part of 3D Gizmos. 
  * Draws outline of 3D Box based on parameters.
  */
-UCLASS(ClassGroup = Utility, HideCategories = (Physics, Collision, Lighting, Rendering, Mobile))
-class INTERACTIVETOOLSFRAMEWORK_API UGizmoBoxComponent : public UGizmoBaseComponent
+UCLASS(ClassGroup = Utility, HideCategories = (Physics, Collision, Lighting, Rendering, Mobile), MinimalAPI)
+class UGizmoBoxComponent : public UGizmoBaseComponent
 {
 	GENERATED_BODY()
 
@@ -37,13 +37,13 @@ public:
 
 private:
 	//~ Begin UPrimitiveComponent Interface.
-	virtual FPrimitiveSceneProxy* CreateSceneProxy() override;
-	virtual bool LineTraceComponent(FHitResult& OutHit, const FVector Start, const FVector End, const FCollisionQueryParams& Params) override;
-	virtual void GetUsedMaterials(TArray<UMaterialInterface*>& OutMaterials, bool bGetDebugMaterials = false) const override;
+	INTERACTIVETOOLSFRAMEWORK_API virtual FPrimitiveSceneProxy* CreateSceneProxy() override;
+	INTERACTIVETOOLSFRAMEWORK_API virtual bool LineTraceComponent(FHitResult& OutHit, const FVector Start, const FVector End, const FCollisionQueryParams& Params) override;
+	INTERACTIVETOOLSFRAMEWORK_API virtual void GetUsedMaterials(TArray<UMaterialInterface*>& OutMaterials, bool bGetDebugMaterials = false) const override;
 	//~ End UPrimitiveComponent Interface.
 
 	//~ Begin USceneComponent Interface.
-	virtual FBoxSphereBounds CalcBounds(const FTransform& LocalToWorld) const override;
+	INTERACTIVETOOLSFRAMEWORK_API virtual FBoxSphereBounds CalcBounds(const FTransform& LocalToWorld) const override;
 	//~ Begin USceneComponent Interface.
 
-};
\ No newline at end of file
+};
diff --git a/Engine/Source/Runtime/InteractiveToolsFramework/Public/BaseGizmos/GizmoCircleComponent.h b/Engine/Source/Runtime/InteractiveToolsFramework/Public/BaseGizmos/GizmoCircleComponent.h
--- a/Engine/Source/Runtime/InteractiveToolsFramework/Public/BaseGizmos/GizmoCircleComponent.h
+++ b/Engine/Source/Runtime/InteractiveToolsFramework/Public/BaseGizmos/GizmoCircleComponent.h
@@ -11,8 +11,8 @@
  * Simple Component intended to be used as part of 3D Gizmos.
  * Draws a 3D circle based on parameters.
  */
-UCLASS(ClassGroup = Utility, HideCategories = (Physics, Collision, Mobile))
-class INTERACTIVETOOLSFRAMEWORK_API UGizmoCircleComponent : public UGizmoBaseComponent
+UCLASS(ClassGroup = Utility, HideCategories = (Physics, Collision, Mobile), MinimalAPI)
+class UGizmoCircleComponent : public UGizmoBaseComponent
 {
 	GENERATED_BODY()
 
@@ -43,11 +43,11 @@ public:
 
 private:
 	//~ Begin UPrimitiveComponent Interface.
-	virtual FPrimitiveSceneProxy* CreateSceneProxy() override;
-	virtual bool LineTraceComponent(FHitResult& OutHit, const FVector Start, const FVector End, const FCollisionQueryParams& Params) override;
+	INTERACTIVETOOLSFRAMEWORK_API virtual FPrimitiveSceneProxy* CreateSceneProxy() override;
+	INTERACTIVETOOLSFRAMEWORK_API virtual bool LineTraceComponent(FHitResult& OutHit, const FVector Start, const FVector End, const FCollisionQueryParams& Params) override;
 	//~ End UPrimitiveComponent Interface.
 
 	//~ Begin USceneComponent Interface.
-	virtual FBoxSphereBounds CalcBounds(const FTransform& LocalToWorld) const override;
+	INTERACTIVETOOLSFRAMEWORK_API virtual FBoxSphereBounds CalcBounds(const FTransform& LocalToWorld) const override;
 	//~ Begin USceneComponent Interface.
-};
\ No newline at end of file
+};
diff --git a/Engine/Source/Runtime/InteractiveToolsFramework/Public/BaseGizmos/GizmoElementArc.h b/Engine/Source/Runtime/InteractiveToolsFramework/Public/BaseGizmos/GizmoElementArc.h
--- a/Engine/Source/Runtime/InteractiveToolsFramework/Public/BaseGizmos/GizmoElementArc.h
+++ b/Engine/Source/Runtime/InteractiveToolsFramework/Public/BaseGizmos/GizmoElementArc.h
@@ -15,25 +15,25 @@ class FMaterialRenderProxy;
  * Simple object intended to be used as part of 3D Gizmos.
  * Draws a thick arc based on parameters.
  */
-UCLASS(Transient)
-class INTERACTIVETOOLSFRAMEWORK_API UGizmoElementArc : public UGizmoElementCircleBase
+UCLASS(Transient, MinimalAPI)
+class UGizmoElementArc : public UGizmoElementCircleBase
 {
 	GENERATED_BODY()
 
 public:
 
 	//~ Begin UGizmoElementBase Interface.
-	virtual void Render(IToolsContextRenderAPI* RenderAPI, const FRenderTraversalState& RenderState) override;
-	virtual FInputRayHit LineTrace(const UGizmoViewContext* ViewContext, const FLineTraceTraversalState& LineTraceState, const FVector& RayOrigin, const FVector& RayDirection) override;
+	INTERACTIVETOOLSFRAMEWORK_API virtual void Render(IToolsContextRenderAPI* RenderAPI, const FRenderTraversalState& RenderState) override;
+	INTERACTIVETOOLSFRAMEWORK_API virtual FInputRayHit LineTrace(const UGizmoViewContext* ViewContext, const FLineTraceTraversalState& LineTraceState, const FVector& RayOrigin, const FVector& RayDirection) override;
 	//~ End UGizmoElementBase Interface.
 
 	// Inner circle radius.
-	virtual void SetInnerRadius(double InInnerRadius);
-	virtual double GetInnerRadius() const;
+	INTERACTIVETOOLSFRAMEWORK_API virtual void SetInnerRadius(double InInnerRadius);
+	INTERACTIVETOOLSFRAMEWORK_API virtual double GetInnerRadius() const;
 
 protected:
 
 	// Arc inner radius.
 	UPROPERTY()
 	double InnerRadius = 5.0;
-};
\ No newline at end of file
+};
diff --git a/Engine/Source/Runtime/InteractiveToolsFramework/Public/BaseGizmos/GizmoElementArrow.h b/Engine/Source/Runtime/InteractiveToolsFramework/Public/BaseGizmos/GizmoElementArrow.h
--- a/Engine/Source/Runtime/InteractiveToolsFramework/Public/BaseGizmos/GizmoElementArrow.h
+++ b/Engine/Source/Runtime/InteractiveToolsFramework/Public/BaseGizmos/GizmoElementArrow.h
@@ -25,65 +25,65 @@ enum class EGizmoElementArrowHeadType
 };
 
 
-UCLASS(Transient)
-class INTERACTIVETOOLSFRAMEWORK_API UGizmoElementArrow : public UGizmoElementBase
+UCLASS(Transient, MinimalAPI)
+class UGizmoElementArrow : public UGizmoElementBase
 {
 	GENERATED_BODY()
 
 public:
-	UGizmoElementArrow();
+	INTERACTIVETOOLSFRAMEWORK_API UGizmoElementArrow();
 
 	//~ Begin UGizmoElementBase Interface.
-	virtual void Render(IToolsContextRenderAPI* RenderAPI, const FRenderTraversalState& RenderState) override;
-	virtual FInputRayHit LineTrace(const UGizmoViewContext* ViewContext, const FLineTraceTraversalState& LineTraceState, const FVector& RayOrigin, const FVector& RayDirection) override;
+	INTERACTIVETOOLSFRAMEWORK_API virtual void Render(IToolsContextRenderAPI* RenderAPI, const FRenderTraversalState& RenderState) override;
+	INTERACTIVETOOLSFRAMEWORK_API virtual FInputRayHit LineTrace(const UGizmoViewContext* ViewContext, const FLineTraceTraversalState& LineTraceState, const FVector& RayOrigin, const FVector& RayDirection) override;
 	//~ End UGizmoElementBase Interface.
 
 	// Location of base of arrow cylinder
-	virtual void SetBase(const FVector& InBase);
-	virtual FVector GetBase() const;
+	INTERACTIVETOOLSFRAMEWORK_API virtual void SetBase(const FVector& InBase);
+	INTERACTIVETOOLSFRAMEWORK_API virtual FVector GetBase() const;
 
 	// Arrow direction.
-	virtual void SetDirection(const FVector& InDirection);
-	virtual FVector GetDirection() const;
+	INTERACTIVETOOLSFRAMEWORK_API virtual void SetDirection(const FVector& InDirection);
+	INTERACTIVETOOLSFRAMEWORK_API virtual FVector GetDirection() const;
 
 	// Arrow side direction for box head.
-	virtual void SetSideDirection(const FVector& InSideDirection);
-	virtual FVector GetSideDirection() const;
+	INTERACTIVETOOLSFRAMEWORK_API virtual void SetSideDirection(const FVector& InSideDirection);
+	INTERACTIVETOOLSFRAMEWORK_API virtual FVector GetSideDirection() const;
 
 	// Arrow body length.
-	virtual void SetBodyLength(float InBodyLength);
-	virtual float GetBodyLength() const;
+	INTERACTIVETOOLSFRAMEWORK_API virtual void SetBodyLength(float InBodyLength);
+	INTERACTIVETOOLSFRAMEWORK_API virtual float GetBodyLength() const;
 
 	// Arrow body cylinder radius.
-	virtual void SetBodyRadius(float InBodyRadius);
-	virtual float GetBodyRadius() const;
+	INTERACTIVETOOLSFRAMEWORK_API virtual void SetBodyRadius(float InBodyRadius);
+	INTERACTIVETOOLSFRAMEWORK_API virtual float GetBodyRadius() const;
 
 	// Arrow head length, used for both cone and cube head.
-	virtual void SetHeadLength(float InHeadLength);
-	virtual float GetHeadLength() const;
+	INTERACTIVETOOLSFRAMEWORK_API virtual void SetHeadLength(float InHeadLength);
+	INTERACTIVETOOLSFRAMEWORK_API virtual float GetHeadLength() const;
 
 	// Arrow head radius, if cone.
-	virtual void SetHeadRadius(float InHeadRadius);
-	virtual float GetHeadRadius() const;
+	INTERACTIVETOOLSFRAMEWORK_API virtual void SetHeadRadius(float InHeadRadius);
+	INTERACTIVETOOLSFRAMEWORK_API virtual float GetHeadRadius() const;
 
 	// Number of sides for cylinder and cone, if relevant.
-	virtual void SetNumSides(int32 InNumSides);
-	virtual int32 GetNumSides() const;
+	INTERACTIVETOOLSFRAMEWORK_API virtual void SetNumSides(int32 InNumSides);
+	INTERACTIVETOOLSFRAMEWORK_API virtual int32 GetNumSides() const;
 
 	// Head type cone or cube.
-	virtual void SetHeadType(EGizmoElementArrowHeadType InHeadType);
-	virtual EGizmoElementArrowHeadType GetHeadType() const;
+	INTERACTIVETOOLSFRAMEWORK_API virtual void SetHeadType(EGizmoElementArrowHeadType InHeadType);
+	INTERACTIVETOOLSFRAMEWORK_API virtual EGizmoElementArrowHeadType GetHeadType() const;
 
 	// Pixel hit distance threshold, element will be scaled enough to add this threshold when line-tracing.
-	virtual void SetPixelHitDistanceThreshold(float InPixelHitDistanceThreshold) override;
+	INTERACTIVETOOLSFRAMEWORK_API virtual void SetPixelHitDistanceThreshold(float InPixelHitDistanceThreshold) override;
 
 protected:
 
 	// Update arrow cylinder based on parameters
-	virtual void UpdateArrowBody();
+	INTERACTIVETOOLSFRAMEWORK_API virtual void UpdateArrowBody();
 
 	// Update arrow cone or box based on parameters
-	virtual void UpdateArrowHead();
+	INTERACTIVETOOLSFRAMEWORK_API virtual void UpdateArrowHead();
 
 	// Flag indicating body properties need to be updated prior to render
 	bool bUpdateArrowBody = true;
diff --git a/Engine/Source/Runtime/InteractiveToolsFramework/Public/BaseGizmos/GizmoElementBase.h b/Engine/Source/Runtime/InteractiveToolsFramework/Public/BaseGizmos/GizmoElementBase.h
--- a/Engine/Source/Runtime/InteractiveToolsFramework/Public/BaseGizmos/GizmoElementBase.h
+++ b/Engine/Source/Runtime/InteractiveToolsFramework/Public/BaseGizmos/GizmoElementBase.h
@@ -21,8 +21,8 @@ class UMaterialInterface;
  * Contains common properties and utility functions.
  * This class does nothing by itself, use subclasses like UGizmoElementCylinder
  */
-UCLASS(Transient, Abstract)
-class INTERACTIVETOOLSFRAMEWORK_API UGizmoElementBase : public UObject
+UCLASS(Transient, Abstract, MinimalAPI)
+class UGizmoElementBase : public UObject
 {
 	GENERATED_BODY()
 public:
@@ -114,118 +114,118 @@ public:
 public:
 
 	// Render enabled visible element.
-	virtual void Render(IToolsContextRenderAPI* RenderAPI, const FRenderTraversalState& RenderState) PURE_VIRTUAL(UGizmoElementBase::Render);
+	INTERACTIVETOOLSFRAMEWORK_API virtual void Render(IToolsContextRenderAPI* RenderAPI, const FRenderTraversalState& RenderState) PURE_VIRTUAL(UGizmoElementBase::Render);
 
 	// Line trace enabled hittable element.
-	virtual FInputRayHit LineTrace(const UGizmoViewContext* ViewContext, const FLineTraceTraversalState& LineTraceState, const FVector& RayOrigin, const FVector& RayDirection) PURE_VIRTUAL(UGizmoElementBase::LineTrace, return FInputRayHit(););
+	INTERACTIVETOOLSFRAMEWORK_API virtual FInputRayHit LineTrace(const UGizmoViewContext* ViewContext, const FLineTraceTraversalState& LineTraceState, const FVector& RayOrigin, const FVector& RayDirection) PURE_VIRTUAL(UGizmoElementBase::LineTrace, return FInputRayHit(););
 
 	// Set/get the visible bit in element state.
-	virtual void SetVisibleState(bool bVisible);
-	virtual bool GetVisibleState() const;
+	INTERACTIVETOOLSFRAMEWORK_API virtual void SetVisibleState(bool bVisible);
+	INTERACTIVETOOLSFRAMEWORK_API virtual bool GetVisibleState() const;
 
 	// Set/get the hittable bit in element state.
-	virtual void SetHittableState(bool bHittable);
-	virtual bool GetHittableState() const;
+	INTERACTIVETOOLSFRAMEWORK_API virtual void SetHittableState(bool bHittable);
+	INTERACTIVETOOLSFRAMEWORK_API virtual bool GetHittableState() const;
 
 	// Element enabled flag. Render and LineTrace only occur when bEnabled is true.
 	// This flag is useful for turning on and off an element globally while retaining its specific settings.
-	virtual void SetEnabled(bool bInEnabled);
-	virtual bool GetEnabled() const;
+	INTERACTIVETOOLSFRAMEWORK_API virtual void SetEnabled(bool bInEnabled);
+	INTERACTIVETOOLSFRAMEWORK_API virtual bool GetEnabled() const;
 
 	// Whether element is enabled for perspective projections.
-	virtual void SetEnabledForPerspectiveProjection(bool bInEnabledForPerspectiveProjection);
-	virtual bool GetEnabledForPerspectiveProjection();
+	INTERACTIVETOOLSFRAMEWORK_API virtual void SetEnabledForPerspectiveProjection(bool bInEnabledForPerspectiveProjection);
+	INTERACTIVETOOLSFRAMEWORK_API virtual bool GetEnabledForPerspectiveProjection();
 
 	// Whether element is enabled for orthographic projections.
-	virtual void SetEnabledInOrthographicProjection(bool bInEnabledForOrthographicProjection);
-	virtual bool GetEnabledInOrthographicProjection();
+	INTERACTIVETOOLSFRAMEWORK_API virtual void SetEnabledInOrthographicProjection(bool bInEnabledForOrthographicProjection);
+	INTERACTIVETOOLSFRAMEWORK_API virtual bool GetEnabledInOrthographicProjection();
 
 	// Whether element is enabled when element state is default.
-	virtual void SetEnabledForDefaultState(bool bInEnabledForDefaultState);
-	virtual bool GetEnabledForDefaultState();
+	INTERACTIVETOOLSFRAMEWORK_API virtual void SetEnabledForDefaultState(bool bInEnabledForDefaultState);
+	INTERACTIVETOOLSFRAMEWORK_API virtual bool GetEnabledForDefaultState();
 
 	// Whether element is enabled when element state is hovering.
-	virtual void SetEnabledForHoveringState(bool bInEnabledForHoveringState);
-	virtual bool GetEnabledForHoveringState();
+	INTERACTIVETOOLSFRAMEWORK_API virtual void SetEnabledForHoveringState(bool bInEnabledForHoveringState);
+	INTERACTIVETOOLSFRAMEWORK_API virtual bool GetEnabledForHoveringState();
 
 	// Whether element is enabled when element state is interacting.
-	virtual void SetEnabledForInteractingState(bool bInEnabledForInteractingState);
-	virtual bool GetEnabledForInteractingState();
+	INTERACTIVETOOLSFRAMEWORK_API virtual void SetEnabledForInteractingState(bool bInEnabledForInteractingState);
+	INTERACTIVETOOLSFRAMEWORK_API virtual bool GetEnabledForInteractingState();
 
 	// For an element hierarchy representing multiple parts of a single gizmo, the part identifier establishes 
 	// a correspondence between a gizmo part and the elements that represent that part. The recognized
 	// part identifier values should be defined in the gizmo. Gizmo part identifiers must be greater than or 
 	// equal to one. Identifier 0 is reserved for the default ID which should be assigned to elements
 	// that do not correspond to any gizmo part, such as non-hittable decorative elements.
-	virtual void SetPartIdentifier(uint32 InPartId);
-	virtual uint32 GetPartIdentifier();
+	INTERACTIVETOOLSFRAMEWORK_API virtual void SetPartIdentifier(uint32 InPartId);
+	INTERACTIVETOOLSFRAMEWORK_API virtual uint32 GetPartIdentifier();
 
 	// Object type bitmask indicating whether this object is visible or hittable or both
-	virtual void SetElementState(EGizmoElementState InElementState);
-	virtual EGizmoElementState GetElementState() const;
+	INTERACTIVETOOLSFRAMEWORK_API virtual void SetElementState(EGizmoElementState InElementState);
+	INTERACTIVETOOLSFRAMEWORK_API virtual EGizmoElementState GetElementState() const;
 
 	// Object interaction state - None, Hovering or Interacting
-	virtual void SetElementInteractionState(EGizmoElementInteractionState InInteractionState);
-	virtual EGizmoElementInteractionState GetElementInteractionState() const;
+	INTERACTIVETOOLSFRAMEWORK_API virtual void SetElementInteractionState(EGizmoElementInteractionState InInteractionState);
+	INTERACTIVETOOLSFRAMEWORK_API virtual EGizmoElementInteractionState GetElementInteractionState() const;
 
 	// Update element's visibility state if element is associated with the specified gizmo part, return true if part was found.
-	virtual bool UpdatePartVisibleState(bool bVisible, uint32 InPartIdentifier);
+	INTERACTIVETOOLSFRAMEWORK_API virtual bool UpdatePartVisibleState(bool bVisible, uint32 InPartIdentifier);
 
 	// Get element's visible state for element associated with the specified gizmo part, if part id was found.
-	virtual TOptional<bool> GetPartVisibleState(uint32 InPartIdentifier) const;
+	INTERACTIVETOOLSFRAMEWORK_API virtual TOptional<bool> GetPartVisibleState(uint32 InPartIdentifier) const;
 
 	// Update element's hittable state if element is associated with the specified gizmo part, return true if part id was found.
-	virtual bool UpdatePartHittableState(bool bHittable, uint32 InPartIdentifier);
+	INTERACTIVETOOLSFRAMEWORK_API virtual bool UpdatePartHittableState(bool bHittable, uint32 InPartIdentifier);
 
 	// Get element's hittable state for element associated with the specified gizmo part, if part id was found.
-	virtual TOptional<bool> GetPartHittableState(uint32 InPartIdentifier) const;
+	INTERACTIVETOOLSFRAMEWORK_API virtual TOptional<bool> GetPartHittableState(uint32 InPartIdentifier) const;
 
 	// Update element's interaction state if element is associated with the specified gizmo part, return true if part id was found.
-	virtual bool UpdatePartInteractionState(EGizmoElementInteractionState InInteractionState, uint32 InPartIdentifier);
+	INTERACTIVETOOLSFRAMEWORK_API virtual bool UpdatePartInteractionState(EGizmoElementInteractionState InInteractionState, uint32 InPartIdentifier);
 
 	// Get element's interaction state for element associated with the specified gizmo part, if part id was found.
-	virtual TOptional<EGizmoElementInteractionState> GetPartInteractionState(uint32 InPartIdentifier) const;
+	INTERACTIVETOOLSFRAMEWORK_API virtual TOptional<EGizmoElementInteractionState> GetPartInteractionState(uint32 InPartIdentifier) const;
 
 	// View-dependent type - None, Axis or Plane. 
-	virtual void SetViewDependentType(EGizmoElementViewDependentType ViewDependentType);
-	virtual EGizmoElementViewDependentType GetViewDependentType() const;
+	INTERACTIVETOOLSFRAMEWORK_API virtual void SetViewDependentType(EGizmoElementViewDependentType ViewDependentType);
+	INTERACTIVETOOLSFRAMEWORK_API virtual EGizmoElementViewDependentType GetViewDependentType() const;
 
 	// View-dependent angle tolerance in radians 
 	//   For Axis, object is culled when angle between view dependent axis and view direction is less than tolerance angle.
 	//   For Planar, cos of angle between view dependent axis (plane normal) and view direction. 
 	// When the view direction is within this tolerance from the plane or axis, this object will be culled.
-	virtual void SetViewDependentAngleTol(float InMaxAngleTol);
-	virtual float GetViewDependentAngleTol() const;
+	INTERACTIVETOOLSFRAMEWORK_API virtual void SetViewDependentAngleTol(float InMaxAngleTol);
+	INTERACTIVETOOLSFRAMEWORK_API virtual float GetViewDependentAngleTol() const;
 
 	// View-dependent axis or plane normal, based on the view-dependent type.
-	virtual void SetViewDependentAxis(FVector InAxis);
-	virtual FVector GetViewDependentAxis() const;
+	INTERACTIVETOOLSFRAMEWORK_API virtual void SetViewDependentAxis(FVector InAxis);
+	INTERACTIVETOOLSFRAMEWORK_API virtual FVector GetViewDependentAxis() const;
 
 	// View align type: None, PointEye, PointOnly, or Axial.
-	virtual void SetViewAlignType(EGizmoElementViewAlignType InViewAlignType);
-	virtual EGizmoElementViewAlignType GetViewAlignType() const;
+	INTERACTIVETOOLSFRAMEWORK_API virtual void SetViewAlignType(EGizmoElementViewAlignType InViewAlignType);
+	INTERACTIVETOOLSFRAMEWORK_API virtual EGizmoElementViewAlignType GetViewAlignType() const;
 
 	// View align axis. 
 	// PointEye, PointScreen and PointOnly rotate this axis to align with view up.
 	// Axial rotates about this axis.
-	virtual void SetViewAlignAxis(FVector InAxis);
-	virtual FVector GetViewAlignAxis() const;
+	INTERACTIVETOOLSFRAMEWORK_API virtual void SetViewAlignAxis(FVector InAxis);
+	INTERACTIVETOOLSFRAMEWORK_API virtual FVector GetViewAlignAxis() const;
 
 	// View align normal.
 	// PointEye rotates the normal to align with camera view direction.
 	// PointScreen rotates the normal to align with screen forward direction.
 	// Axial rotates the normal around the axis to align as closely as possible with the view direction.
-	virtual void SetViewAlignNormal(FVector InAxis);
-	virtual FVector GetViewAlignNormal() const;
+	INTERACTIVETOOLSFRAMEWORK_API virtual void SetViewAlignNormal(FVector InAxis);
+	INTERACTIVETOOLSFRAMEWORK_API virtual FVector GetViewAlignNormal() const;
 
 	// View-align angle tolerance in radians.
 	// Viewer alignment will not occur when the viewing angle is within this angle of view align axis.
-	virtual void SetViewAlignAxialAngleTol(float InMaxAngleTol);
-	virtual float GetViewAlignAxialAngleTol() const;
+	INTERACTIVETOOLSFRAMEWORK_API virtual void SetViewAlignAxialAngleTol(float InMaxAngleTol);
+	INTERACTIVETOOLSFRAMEWORK_API virtual float GetViewAlignAxialAngleTol() const;
 
 	// Pixel hit distance threshold, element will be scaled enough to add this threshold when line-tracing. */
-	virtual void SetPixelHitDistanceThreshold(float InPixelHitDistanceThreshold);
-	virtual float GetPixelHitDistanceThreshold() const;
+	INTERACTIVETOOLSFRAMEWORK_API virtual void SetPixelHitDistanceThreshold(float InPixelHitDistanceThreshold);
+	INTERACTIVETOOLSFRAMEWORK_API virtual float GetPixelHitDistanceThreshold() const;
 
 	//
 	// Methods for managing render state attributes: Material, HoverMaterial, InteractMaterial, VertexColor 
@@ -239,95 +239,95 @@ public:
 	// Set mesh render state material attribute. 
 	//  @param InMaterial - material to be set
 	//  @param InOverridesChildState - when true, this material will override the material of all child elements.
-	virtual void SetMaterial(TWeakObjectPtr<UMaterialInterface> InMaterial, bool InOverridesChildState = false);
+	INTERACTIVETOOLSFRAMEWORK_API virtual void SetMaterial(TWeakObjectPtr<UMaterialInterface> InMaterial, bool InOverridesChildState = false);
 
 	// Get mesh render state material attribute's value. 
-	virtual const UMaterialInterface* GetMaterial() const;
+	INTERACTIVETOOLSFRAMEWORK_API virtual const UMaterialInterface* GetMaterial() const;
 
 	// Get mesh render state material attribute's override setting. 
-	virtual bool DoesMaterialOverrideChildState() const;
+	INTERACTIVETOOLSFRAMEWORK_API virtual bool DoesMaterialOverrideChildState() const;
 
 	// Clear mesh render state material attribute. 
-	virtual void ClearMaterial();
+	INTERACTIVETOOLSFRAMEWORK_API virtual void ClearMaterial();
 
 	// Set mesh render state hover material attribute. 
 	//  @param InHoverMaterial - hover material to be set
 	//  @param InOverridesChildState - when true, this hover material will override the material of all child elements.
-	virtual void SetHoverMaterial(TWeakObjectPtr<UMaterialInterface> InHoverMaterial, bool InOverridesChildState = false);
+	INTERACTIVETOOLSFRAMEWORK_API virtual void SetHoverMaterial(TWeakObjectPtr<UMaterialInterface> InHoverMaterial, bool InOverridesChildState = false);
 
 	// Get mesh render state hover material attribute's value.
-	virtual const UMaterialInterface* GetHoverMaterial() const;
+	INTERACTIVETOOLSFRAMEWORK_API virtual const UMaterialInterface* GetHoverMaterial() const;
 
 	// Get mesh render state hover material attribute's override setting.
-	virtual bool DoesHoverMaterialOverrideChildState() const;
+	INTERACTIVETOOLSFRAMEWORK_API virtual bool DoesHoverMaterialOverrideChildState() const;
 
 	// Clear mesh render state hover material attribute. 
-	virtual void ClearHoverMaterial();
+	INTERACTIVETOOLSFRAMEWORK_API virtual void ClearHoverMaterial();
 
 	// Set mesh render state interact material attribute. 
 	//  @param InHoverMaterial - interact material to be set
 	//  @param InOverridesChildState - when true, this interact material will override the material of all child elements.
-	virtual void SetInteractMaterial(TWeakObjectPtr<UMaterialInterface> InInteractMaterial, bool InOverridesChildState = false);
+	INTERACTIVETOOLSFRAMEWORK_API virtual void SetInteractMaterial(TWeakObjectPtr<UMaterialInterface> InInteractMaterial, bool InOverridesChildState = false);
 
 	// Get mesh render state interact material attribute's value. 
-	virtual const UMaterialInterface* GetInteractMaterial() const;
+	INTERACTIVETOOLSFRAMEWORK_API virtual const UMaterialInterface* GetInteractMaterial() const;
 
 	// Get mesh render state interact material attribute's override setting. 
-	virtual bool DoesInteractMaterialOverrideChildState() const;
+	INTERACTIVETOOLSFRAMEWORK_API virtual bool DoesInteractMaterialOverrideChildState() const;
 
 	// Clear mesh render interact state material attribute. 
-	virtual void ClearInteractMaterial();
+	INTERACTIVETOOLSFRAMEWORK_API virtual void ClearInteractMaterial();
 
 	// Set mesh render state vertex color attribute. 
 	//  @param InVertexColor - vertex color to be set
 	//  @param InOverridesChildState - when true, this vertex color will override the material of all child elements.
-	virtual void SetVertexColor(FLinearColor InVertexColor, bool InOverridesChildState = false);
+	INTERACTIVETOOLSFRAMEWORK_API virtual void SetVertexColor(FLinearColor InVertexColor, bool InOverridesChildState = false);
 
 	// Get mesh render state vertex color attribute's value. 
-	virtual FLinearColor GetVertexColor() const;
+	INTERACTIVETOOLSFRAMEWORK_API virtual FLinearColor GetVertexColor() const;
 
 	// Returns true, if mesh render state vertex color attribute has been set. 
-	virtual bool HasVertexColor() const;
+	INTERACTIVETOOLSFRAMEWORK_API virtual bool HasVertexColor() const;
 
 	// Get mesh render state vertex color attribute's override setting. 
-	virtual bool DoesVertexColorOverrideChildState() const;
+	INTERACTIVETOOLSFRAMEWORK_API virtual bool DoesVertexColorOverrideChildState() const;
 
 	// Clear mesh render state vertex color attribute.
-	virtual void ClearVertexColor();
+	INTERACTIVETOOLSFRAMEWORK_API virtual void ClearVertexColor();
 
 	// Set mesh render state vertex color attribute. 
 	//  @param InVertexColor - vertex color to be set
 	//  @param InOverridesChildState - when true, this vertex color will override the material of all child elements.
-	virtual void SetHoverVertexColor(FLinearColor InVertexColor, bool InOverridesChildState = false);
+	INTERACTIVETOOLSFRAMEWORK_API virtual void SetHoverVertexColor(FLinearColor InVertexColor, bool InOverridesChildState = false);
 
 	// Get mesh render state vertex color attribute's value. 
-	virtual FLinearColor GetHoverVertexColor() const;
+	INTERACTIVETOOLSFRAMEWORK_API virtual FLinearColor GetHoverVertexColor() const;
 
 	// Returns true, if mesh render state vertex color attribute has been set. 
-	virtual bool HasHoverVertexColor() const;
+	INTERACTIVETOOLSFRAMEWORK_API virtual bool HasHoverVertexColor() const;
 
 	// Get mesh render state vertex color attribute's override setting. 
-	virtual bool DoesHoverVertexColorOverrideChildState() const;
+	INTERACTIVETOOLSFRAMEWORK_API virtual bool DoesHoverVertexColorOverrideChildState() const;
 
 	// Clear mesh render state vertex color attribute.
-	virtual void ClearHoverVertexColor();
+	INTERACTIVETOOLSFRAMEWORK_API virtual void ClearHoverVertexColor();
 
 	// Set mesh render state vertex color attribute. 
 	//  @param InVertexColor - vertex color to be set
 	//  @param InOverridesChildState - when true, this vertex color will override the material of all child elements.
-	virtual void SetInteractVertexColor(FLinearColor InVertexColor, bool InOverridesChildState = false);
+	INTERACTIVETOOLSFRAMEWORK_API virtual void SetInteractVertexColor(FLinearColor InVertexColor, bool InOverridesChildState = false);
 
 	// Get mesh render state vertex color attribute's value. 
-	virtual FLinearColor GetInteractVertexColor() const;
+	INTERACTIVETOOLSFRAMEWORK_API virtual FLinearColor GetInteractVertexColor() const;
 
 	// Returns true, if mesh render state vertex color attribute has been set. 
-	virtual bool HasInteractVertexColor() const;
+	INTERACTIVETOOLSFRAMEWORK_API virtual bool HasInteractVertexColor() const;
 
 	// Get mesh render state vertex color attribute's override setting. 
-	virtual bool DoesInteractVertexColorOverrideChildState() const;
+	INTERACTIVETOOLSFRAMEWORK_API virtual bool DoesInteractVertexColorOverrideChildState() const;
 
 	// Clear mesh render state vertex color attribute.
-	virtual void ClearInteractVertexColor();
+	INTERACTIVETOOLSFRAMEWORK_API virtual void ClearInteractVertexColor();
 
 protected:
 
@@ -430,73 +430,73 @@ protected:
 protected:
 
 	// Return whether element is currently visible.
-	virtual bool IsVisible(const FSceneView* View, EGizmoElementInteractionState InCurrentInteractionState, 
+	INTERACTIVETOOLSFRAMEWORK_API virtual bool IsVisible(const FSceneView* View, EGizmoElementInteractionState InCurrentInteractionState, 
 		const FTransform& InLocalToWorldTransform, const FVector& InLocalCenter) const;
 
 	// Return whether element is currently hittable.
-	virtual bool IsHittable(const UGizmoViewContext* ViewContext, const FTransform& InLocalToWorldTransform, const FVector& InLocalCenter) const;
+	INTERACTIVETOOLSFRAMEWORK_API virtual bool IsHittable(const UGizmoViewContext* ViewContext, const FTransform& InLocalToWorldTransform, const FVector& InLocalCenter) const;
 
 	// Returns whether object is visible in input FSceneView based on view-dependent visibility settings.
-	virtual bool GetViewDependentVisibility(const FSceneView* View, const FTransform& InLocalToWorldTransform, const FVector& InLocalCenter) const;
+	INTERACTIVETOOLSFRAMEWORK_API virtual bool GetViewDependentVisibility(const FSceneView* View, const FTransform& InLocalToWorldTransform, const FVector& InLocalCenter) const;
 
 	// Returns whether object is visible in input gizmo view context based on view-dependent visibility settings
-	virtual bool GetViewDependentVisibility(const UGizmoViewContext* View, const FTransform& InLocalToWorldTransform, const FVector& InLocalCenter) const;
+	INTERACTIVETOOLSFRAMEWORK_API virtual bool GetViewDependentVisibility(const UGizmoViewContext* View, const FTransform& InLocalToWorldTransform, const FVector& InLocalCenter) const;
 
 	// Returns whether object is visible based on view-dependent visibility settings.
-	virtual bool GetViewDependentVisibility(const FVector& InViewLocation, const FVector& InViewDirection, bool bInPerspectiveView, const FTransform& InLocalToWorldTransform, const FVector& InLocalCenter) const;
+	INTERACTIVETOOLSFRAMEWORK_API virtual bool GetViewDependentVisibility(const FVector& InViewLocation, const FVector& InViewDirection, bool bInPerspectiveView, const FTransform& InLocalToWorldTransform, const FVector& InLocalCenter) const;
 
 	// Returns whether element is enabled for given interaction state.
-	virtual bool GetEnabledForInteractionState(EGizmoElementInteractionState InInteractionState) const;
+	INTERACTIVETOOLSFRAMEWORK_API virtual bool GetEnabledForInteractionState(EGizmoElementInteractionState InInteractionState) const;
 
 	// Returns whether element is enabled for given view projection type.
-	virtual bool GetEnabledForViewProjection(bool bIsPerspectiveProjection) const;
+	INTERACTIVETOOLSFRAMEWORK_API virtual bool GetEnabledForViewProjection(bool bIsPerspectiveProjection) const;
 
 	// Return whether this element has a view alignment rotation based on input FSceneView and view-dependent alignment settings.
 	// @param OutAlignRot the rotation to align this element in local space, should be prepended to the local-to-world transform.
-	virtual bool GetViewAlignRot(const FSceneView* View, const FTransform& InLocalToWorldTransform, const FVector& InLocalCenter, FQuat& OutAlignRot) const;
+	INTERACTIVETOOLSFRAMEWORK_API virtual bool GetViewAlignRot(const FSceneView* View, const FTransform& InLocalToWorldTransform, const FVector& InLocalCenter, FQuat& OutAlignRot) const;
 
 	// Return whether this element has a view alignment rotation based on input gizmo view context and view-dependent alignment settings.
 	// @param OutAlignRot the rotation to align this element in local space, should be prepended to the local-to-world transform.
-	virtual bool GetViewAlignRot(const UGizmoViewContext* View, const FTransform& InLocalToWorldTransform, const FVector& InLocalCenter, FQuat& OutAlignRot) const;
+	INTERACTIVETOOLSFRAMEWORK_API virtual bool GetViewAlignRot(const UGizmoViewContext* View, const FTransform& InLocalToWorldTransform, const FVector& InLocalCenter, FQuat& OutAlignRot) const;
 
 	// Return whether this element has a view alignment rotation based on input view parameters and view-dependent alignment settings.
 	// @param OutAlignRot the rotation to align this element in local space, should be prepended to the local-to-world transform.
-	virtual bool GetViewAlignRot(const FVector& InViewLocation, const FVector& InViewDirection, const FVector& InViewUp, bool bInPerspectiveView, const FTransform& InLocalToWorldTransform, const FVector& InLocalCenter, FQuat& OutAlignRot) const;
+	INTERACTIVETOOLSFRAMEWORK_API virtual bool GetViewAlignRot(const FVector& InViewLocation, const FVector& InViewDirection, const FVector& InViewUp, bool bInPerspectiveView, const FTransform& InLocalToWorldTransform, const FVector& InLocalCenter, FQuat& OutAlignRot) const;
 
 	// Update render state during render traversal, determines the current render state for this element 
 	// @param RenderAPI - tools render context
 	// @param InLocalCenter - local element center position
 	// @param InOutRenderState - render state's local to world transform will be updated with translation to center and view-dependent alignment rotation, if applicable.
 	// @return view dependent visibility, true if this element is visible in the current view. 
-	virtual bool UpdateRenderState(IToolsContextRenderAPI* RenderAPI, const FVector& InLocalCenter, FRenderTraversalState& InOutRenderState);
+	INTERACTIVETOOLSFRAMEWORK_API virtual bool UpdateRenderState(IToolsContextRenderAPI* RenderAPI, const FVector& InLocalCenter, FRenderTraversalState& InOutRenderState);
 
 	// Update render state during render traversal, determines the current render state for this element
 	// Same parameters as UpdateRenderState above plus two output parameters:
 	// @param bOutHasAlignRot - whether alignment rotation was applied to output state
 	// @param OutAlignRot - alignment rotation applied to output state, if applicable.
-	virtual bool UpdateRenderState(IToolsContextRenderAPI* RenderAPI, const FVector& InLocalCenter, FRenderTraversalState& InOutRenderState, bool& bOutHasAlignRot, FQuat& OutAlignRot);
+	INTERACTIVETOOLSFRAMEWORK_API virtual bool UpdateRenderState(IToolsContextRenderAPI* RenderAPI, const FVector& InLocalCenter, FRenderTraversalState& InOutRenderState, bool& bOutHasAlignRot, FQuat& OutAlignRot);
 
 	// Update line trace state during line trace traversal, determines the current state for this element 
 	// @param ViewContext - current gizmo view context
 	// @param InLocalCenter - local element center position, this will update InOutLineTraceState's LocalToWorldTransform with a translation to LocalCenter
 	// @param InOutLineTraceState - line trace state's local to world transform will be updated with translation to center and view-dependent alignment rotation, if applicable.
 	// @return view dependent visibility, true if this element is visible in the current view. 
-	virtual bool UpdateLineTraceState(const UGizmoViewContext* ViewContext, const FVector& InLocalCenter, FLineTraceTraversalState& InOutLineTraceState);
+	INTERACTIVETOOLSFRAMEWORK_API virtual bool UpdateLineTraceState(const UGizmoViewContext* ViewContext, const FVector& InLocalCenter, FLineTraceTraversalState& InOutLineTraceState);
 
 	// Update line trace state during line trace traversal, determines the current state for this element 
 	// Same parameters as UpdateLineTraceState above plus two output parameters:
 	// @param bOutHasAlignRot - whether alignment rotation was applied to output state
 	// @param OutAlignRot - alignment rotation applied to output state, if applicable.
-	virtual bool UpdateLineTraceState(const UGizmoViewContext* ViewContext, const FVector& InLocalCenter, FLineTraceTraversalState& InOutRenderState, bool& bOutHasAlignRot, FQuat& OutAlignRot);
+	INTERACTIVETOOLSFRAMEWORK_API virtual bool UpdateLineTraceState(const UGizmoViewContext* ViewContext, const FVector& InLocalCenter, FLineTraceTraversalState& InOutRenderState, bool& bOutHasAlignRot, FQuat& OutAlignRot);
 
 protected:
 
 	// Helper method to verify scale is uniform. If it is non-uniform, a one-time warning log is issued.
 	// Returns true if scale is uniform.
-	bool VerifyUniformScale(const FVector& Scale) const;
+	INTERACTIVETOOLSFRAMEWORK_API bool VerifyUniformScale(const FVector& Scale) const;
 
 	// Helper method for view alignment.
 	// Returns rotation between source and target input coordinate spaces.
-	FQuat GetAlignRotBetweenCoordSpaces(FVector SourceForward, FVector SourceRight, FVector SourceUp, FVector TargetForward, FVector TargetRight, FVector TargetUp) const;
+	INTERACTIVETOOLSFRAMEWORK_API FQuat GetAlignRotBetweenCoordSpaces(FVector SourceForward, FVector SourceRight, FVector SourceUp, FVector TargetForward, FVector TargetRight, FVector TargetUp) const;
 };
 
diff --git a/Engine/Source/Runtime/InteractiveToolsFramework/Public/BaseGizmos/GizmoElementBox.h b/Engine/Source/Runtime/InteractiveToolsFramework/Public/BaseGizmos/GizmoElementBox.h
--- a/Engine/Source/Runtime/InteractiveToolsFramework/Public/BaseGizmos/GizmoElementBox.h
+++ b/Engine/Source/Runtime/InteractiveToolsFramework/Public/BaseGizmos/GizmoElementBox.h
@@ -13,33 +13,33 @@
  * Draws a solid 3D cylinder based on parameters.
  */
 
-UCLASS(Transient)
-class INTERACTIVETOOLSFRAMEWORK_API UGizmoElementBox : public UGizmoElementBase
+UCLASS(Transient, MinimalAPI)
+class UGizmoElementBox : public UGizmoElementBase
 {
 	GENERATED_BODY()
 
 public:
 
 	//~ Begin UGizmoElementBase Interface.
-	virtual void Render(IToolsContextRenderAPI* RenderAPI, const FRenderTraversalState& RenderState) override;
-	virtual FInputRayHit LineTrace(const UGizmoViewContext* ViewContext, const FLineTraceTraversalState& LineTraceState, const FVector& RayOrigin, const FVector& RayDirection) override;
+	INTERACTIVETOOLSFRAMEWORK_API virtual void Render(IToolsContextRenderAPI* RenderAPI, const FRenderTraversalState& RenderState) override;
+	INTERACTIVETOOLSFRAMEWORK_API virtual FInputRayHit LineTrace(const UGizmoViewContext* ViewContext, const FLineTraceTraversalState& LineTraceState, const FVector& RayOrigin, const FVector& RayDirection) override;
 	//~ End UGizmoElementBase Interface.
 
 	// Location of center of box's base circle.
-	virtual void SetCenter(const FVector& InCenter);
-	virtual FVector GetCenter() const;
+	INTERACTIVETOOLSFRAMEWORK_API virtual void SetCenter(const FVector& InCenter);
+	INTERACTIVETOOLSFRAMEWORK_API virtual FVector GetCenter() const;
 
 	// Box axis up direction.
-	virtual void SetUpDirection(const FVector& InUpDirection);
-	virtual FVector GetUpDirection() const;
+	INTERACTIVETOOLSFRAMEWORK_API virtual void SetUpDirection(const FVector& InUpDirection);
+	INTERACTIVETOOLSFRAMEWORK_API virtual FVector GetUpDirection() const;
 
 	// Box axis side direction.
-	virtual void SetSideDirection(const FVector& InSideDirection);
-	virtual FVector GetSideDirection() const;
+	INTERACTIVETOOLSFRAMEWORK_API virtual void SetSideDirection(const FVector& InSideDirection);
+	INTERACTIVETOOLSFRAMEWORK_API virtual FVector GetSideDirection() const;
 
 	// Box dimensions.
-	virtual void SetDimensions(const FVector& InDimensions);
-	virtual FVector GetDimensions() const;
+	INTERACTIVETOOLSFRAMEWORK_API virtual void SetDimensions(const FVector& InDimensions);
+	INTERACTIVETOOLSFRAMEWORK_API virtual FVector GetDimensions() const;
 
 protected:
 
diff --git a/Engine/Source/Runtime/InteractiveToolsFramework/Public/BaseGizmos/GizmoElementCircle.h b/Engine/Source/Runtime/InteractiveToolsFramework/Public/BaseGizmos/GizmoElementCircle.h
--- a/Engine/Source/Runtime/InteractiveToolsFramework/Public/BaseGizmos/GizmoElementCircle.h
+++ b/Engine/Source/Runtime/InteractiveToolsFramework/Public/BaseGizmos/GizmoElementCircle.h
@@ -14,32 +14,32 @@
  * 
  * The circle element does not yet support partial circles.
  */
-UCLASS(Transient)
-class INTERACTIVETOOLSFRAMEWORK_API UGizmoElementCircle : public UGizmoElementCircleBase
+UCLASS(Transient, MinimalAPI)
+class UGizmoElementCircle : public UGizmoElementCircleBase
 {
 	GENERATED_BODY()
 
 public:
 	//~ Begin UGizmoElementBase Interface.
-	virtual void Render(IToolsContextRenderAPI* RenderAPI, const FRenderTraversalState& RenderState) override;
-	virtual FInputRayHit LineTrace(const UGizmoViewContext* ViewContext, const FLineTraceTraversalState& LineTraceState, const FVector& RayOrigin, const FVector& RayDirection) override;
+	INTERACTIVETOOLSFRAMEWORK_API virtual void Render(IToolsContextRenderAPI* RenderAPI, const FRenderTraversalState& RenderState) override;
+	INTERACTIVETOOLSFRAMEWORK_API virtual FInputRayHit LineTrace(const UGizmoViewContext* ViewContext, const FLineTraceTraversalState& LineTraceState, const FVector& RayOrigin, const FVector& RayDirection) override;
 	//~ End UGizmoElementBase Interface.
 
 	// Draw mesh
-	virtual void SetDrawMesh(bool InDrawMesh);
-	virtual bool GetDrawMesh() const;
+	INTERACTIVETOOLSFRAMEWORK_API virtual void SetDrawMesh(bool InDrawMesh);
+	INTERACTIVETOOLSFRAMEWORK_API virtual bool GetDrawMesh() const;
 
 	// Draw line
-	virtual void SetDrawLine(bool InDrawLine);
-	virtual bool GetDrawLine() const;
+	INTERACTIVETOOLSFRAMEWORK_API virtual void SetDrawLine(bool InDrawLine);
+	INTERACTIVETOOLSFRAMEWORK_API virtual bool GetDrawLine() const;
 
 	// Hit mesh
-	virtual void SetHitMesh(bool InHitMesh);
-	virtual bool GetHitMesh() const;
+	INTERACTIVETOOLSFRAMEWORK_API virtual void SetHitMesh(bool InHitMesh);
+	INTERACTIVETOOLSFRAMEWORK_API virtual bool GetHitMesh() const;
 
 	// Hit line
-	virtual void SetHitLine(bool InHitLine);
-	virtual bool GetHitLine() const;
+	INTERACTIVETOOLSFRAMEWORK_API virtual void SetHitLine(bool InHitLine);
+	INTERACTIVETOOLSFRAMEWORK_API virtual bool GetHitLine() const;
 
 protected:
 
@@ -58,4 +58,4 @@ protected:
 	// Whether to perform hit test on line.
 	UPROPERTY()
 	bool bHitLine = false;
-};
\ No newline at end of file
+};
