[Engine]

diff --git a/Engine/Source/Runtime/InteractiveToolsFramework/Public/BaseGizmos/GizmoElementCircleBase.h b/Engine/Source/Runtime/InteractiveToolsFramework/Public/BaseGizmos/GizmoElementCircleBase.h
--- a/Engine/Source/Runtime/InteractiveToolsFramework/Public/BaseGizmos/GizmoElementCircleBase.h
+++ b/Engine/Source/Runtime/InteractiveToolsFramework/Public/BaseGizmos/GizmoElementCircleBase.h
@@ -14,51 +14,51 @@ class FMaterialRenderProxy;
 /**
  * Abstract base object for circle, torus and arc.
  */
-UCLASS(Transient, Abstract)
-class INTERACTIVETOOLSFRAMEWORK_API UGizmoElementCircleBase : public UGizmoElementLineBase
+UCLASS(Transient, Abstract, MinimalAPI)
+class UGizmoElementCircleBase : public UGizmoElementLineBase
 {
 	GENERATED_BODY()
 
 public:
 
 	// CircleBase center.
-	virtual void SetCenter(const FVector& InCenter);
-	virtual FVector GetCenter() const;
+	INTERACTIVETOOLSFRAMEWORK_API virtual void SetCenter(const FVector& InCenter);
+	INTERACTIVETOOLSFRAMEWORK_API virtual FVector GetCenter() const;
 
 	// Axis0 of plane in which circle lies, must perpendicular to normal. 
 	// Start and end angles for partial circles are relative to this axis.
-	virtual void SetAxis0(const FVector& InAxis0);
-	virtual FVector GetAxis0() const;
+	INTERACTIVETOOLSFRAMEWORK_API virtual void SetAxis0(const FVector& InAxis0);
+	INTERACTIVETOOLSFRAMEWORK_API virtual FVector GetAxis0() const;
 
 	// Axis1 of plane in which circle lies, must perpendicular to Axis0.
-	virtual void SetAxis1(const FVector& InAxis0);
-	virtual FVector GetAxis1() const;
+	INTERACTIVETOOLSFRAMEWORK_API virtual void SetAxis1(const FVector& InAxis0);
+	INTERACTIVETOOLSFRAMEWORK_API virtual FVector GetAxis1() const;
 
 	// Circle radius.
-	virtual void SetRadius(double InRadius);
-	virtual double GetRadius() const;
+	INTERACTIVETOOLSFRAMEWORK_API virtual void SetRadius(double InRadius);
+	INTERACTIVETOOLSFRAMEWORK_API virtual double GetRadius() const;
 
 	// Number of segments for rendering circle.
-	virtual void SetNumSegments(int32 InNumSegments);
-	virtual int32 GetNumSegments() const;
+	INTERACTIVETOOLSFRAMEWORK_API virtual void SetNumSegments(int32 InNumSegments);
+	INTERACTIVETOOLSFRAMEWORK_API virtual int32 GetNumSegments() const;
 
 	// When Partial, renders partial arc based on angle.
 	// When PartialViewDependent, same as Partial but renders full arc when arc normal aligns with view direction
-	virtual void SetPartialType(EGizmoElementPartialType InPartial);
-	virtual EGizmoElementPartialType GetPartialType() const;
+	INTERACTIVETOOLSFRAMEWORK_API virtual void SetPartialType(EGizmoElementPartialType InPartial);
+	INTERACTIVETOOLSFRAMEWORK_API virtual EGizmoElementPartialType GetPartialType() const;
 
 	// Start of arc angle of partial torus in radians, relative to Axis0.
-	virtual void SetPartialStartAngle(double InPartialAngle);
-	virtual double GetPartialStartAngle() const;
+	INTERACTIVETOOLSFRAMEWORK_API virtual void SetPartialStartAngle(double InPartialAngle);
+	INTERACTIVETOOLSFRAMEWORK_API virtual double GetPartialStartAngle() const;
 
 	// Start of arc angle of partial torus in radians, relative to Axis1.
-	virtual void SetPartialEndAngle(double InPartialAngle);
-	virtual double GetPartialEndAngle() const;
+	INTERACTIVETOOLSFRAMEWORK_API virtual void SetPartialEndAngle(double InPartialAngle);
+	INTERACTIVETOOLSFRAMEWORK_API virtual double GetPartialEndAngle() const;
 
 	// If partial type is PartialViewDependent, when the cosine of angle between the normal and view direction 
 	// is within this tolerance, the arc will be rendered as full rather than partial
-	virtual void SetPartialViewDependentMaxCosTol(double InPartialViewDependentMaxCosTol);
-	virtual double GetPartialViewDependentMaxCosTol() const;
+	INTERACTIVETOOLSFRAMEWORK_API virtual void SetPartialViewDependentMaxCosTol(double InPartialViewDependentMaxCosTol);
+	INTERACTIVETOOLSFRAMEWORK_API virtual double GetPartialViewDependentMaxCosTol() const;
 
 protected:
 
@@ -103,13 +103,13 @@ protected:
 protected:
 
 	// Returns whether element should be partial based on current view
-	bool IsPartial(const FSceneView* View, const FVector& InWorldCenter, const FVector& InWorldNormal);
+	INTERACTIVETOOLSFRAMEWORK_API bool IsPartial(const FSceneView* View, const FVector& InWorldCenter, const FVector& InWorldNormal);
 
 	// Returns whether element should be partial based on current view
-	bool IsPartial(const UGizmoViewContext* ViewContext, const FVector& InWorldCenter, const FVector& InWorldNormal);
+	INTERACTIVETOOLSFRAMEWORK_API bool IsPartial(const UGizmoViewContext* ViewContext, const FVector& InWorldCenter, const FVector& InWorldNormal);
 
 	// Returns whether element should be partial based on current view
-	bool IsPartial(const FVector& InWorldCenter, const FVector& InWorldNormal,
+	INTERACTIVETOOLSFRAMEWORK_API bool IsPartial(const FVector& InWorldCenter, const FVector& InWorldNormal,
 		const FVector& InViewLocation, const FVector& InViewDirection,
 		bool bIsPerspectiveProjection);
-};
\ No newline at end of file
+};
diff --git a/Engine/Source/Runtime/InteractiveToolsFramework/Public/BaseGizmos/GizmoElementCone.h b/Engine/Source/Runtime/InteractiveToolsFramework/Public/BaseGizmos/GizmoElementCone.h
--- a/Engine/Source/Runtime/InteractiveToolsFramework/Public/BaseGizmos/GizmoElementCone.h
+++ b/Engine/Source/Runtime/InteractiveToolsFramework/Public/BaseGizmos/GizmoElementCone.h
@@ -13,37 +13,37 @@
  * Draws a solid 3D cone based on parameters.
  */
 
-UCLASS(Transient)
-class INTERACTIVETOOLSFRAMEWORK_API UGizmoElementCone : public UGizmoElementBase
+UCLASS(Transient, MinimalAPI)
+class UGizmoElementCone : public UGizmoElementBase
 {
 	GENERATED_BODY()
 
 public:
 
 	//~ Begin UGizmoElementBase Interface.
-	virtual void Render(IToolsContextRenderAPI* RenderAPI, const FRenderTraversalState& RenderState) override;
-	virtual FInputRayHit LineTrace(const UGizmoViewContext* ViewContext, const FLineTraceTraversalState& LineTraceState, const FVector& RayOrigin, const FVector& RayDirection) override;
+	INTERACTIVETOOLSFRAMEWORK_API virtual void Render(IToolsContextRenderAPI* RenderAPI, const FRenderTraversalState& RenderState) override;
+	INTERACTIVETOOLSFRAMEWORK_API virtual FInputRayHit LineTrace(const UGizmoViewContext* ViewContext, const FLineTraceTraversalState& LineTraceState, const FVector& RayOrigin, const FVector& RayDirection) override;
 	//~ End UGizmoElementBase Interface.
 
 	// Cone tip location.
-	virtual void SetOrigin(const FVector& InBase);
-	virtual FVector GetOrigin() const;
+	INTERACTIVETOOLSFRAMEWORK_API virtual void SetOrigin(const FVector& InBase);
+	INTERACTIVETOOLSFRAMEWORK_API virtual FVector GetOrigin() const;
 
 	// Cone axis direction pointing from tip toward base of cone.
-	virtual void SetDirection(const FVector& InDirection);
-	virtual FVector GetDirection() const;
+	INTERACTIVETOOLSFRAMEWORK_API virtual void SetDirection(const FVector& InDirection);
+	INTERACTIVETOOLSFRAMEWORK_API virtual FVector GetDirection() const;
 
 	// Cone height.
-	virtual void SetHeight(float InHeight);
-	virtual float GetHeight() const;
+	INTERACTIVETOOLSFRAMEWORK_API virtual void SetHeight(float InHeight);
+	INTERACTIVETOOLSFRAMEWORK_API virtual float GetHeight() const;
 
 	// Cone radius.
-	virtual void SetRadius(float InRadius);
-	virtual float GetRadius() const;
+	INTERACTIVETOOLSFRAMEWORK_API virtual void SetRadius(float InRadius);
+	INTERACTIVETOOLSFRAMEWORK_API virtual float GetRadius() const;
 
 	// Number of sides for tessellating cone.
-	virtual void SetNumSides(int32 InNumSides);
-	virtual int32 GetNumSides() const;
+	INTERACTIVETOOLSFRAMEWORK_API virtual void SetNumSides(int32 InNumSides);
+	INTERACTIVETOOLSFRAMEWORK_API virtual int32 GetNumSides() const;
 
 protected:
 
diff --git a/Engine/Source/Runtime/InteractiveToolsFramework/Public/BaseGizmos/GizmoElementCylinder.h b/Engine/Source/Runtime/InteractiveToolsFramework/Public/BaseGizmos/GizmoElementCylinder.h
--- a/Engine/Source/Runtime/InteractiveToolsFramework/Public/BaseGizmos/GizmoElementCylinder.h
+++ b/Engine/Source/Runtime/InteractiveToolsFramework/Public/BaseGizmos/GizmoElementCylinder.h
@@ -13,37 +13,37 @@
  * Draws a solid 3D cylinder based on parameters.
  */
 
-UCLASS(Transient)
-class INTERACTIVETOOLSFRAMEWORK_API UGizmoElementCylinder : public UGizmoElementBase
+UCLASS(Transient, MinimalAPI)
+class UGizmoElementCylinder : public UGizmoElementBase
 {
 	GENERATED_BODY()
 
 public:
 
 	//~ Begin UGizmoElementBase Interface.
-	virtual void Render(IToolsContextRenderAPI* RenderAPI, const FRenderTraversalState& RenderState) override;
-	virtual FInputRayHit LineTrace(const UGizmoViewContext* ViewContext, const FLineTraceTraversalState& LineTraceState, const FVector& RayOrigin, const FVector& RayDirection) override;
+	INTERACTIVETOOLSFRAMEWORK_API virtual void Render(IToolsContextRenderAPI* RenderAPI, const FRenderTraversalState& RenderState) override;
+	INTERACTIVETOOLSFRAMEWORK_API virtual FInputRayHit LineTrace(const UGizmoViewContext* ViewContext, const FLineTraceTraversalState& LineTraceState, const FVector& RayOrigin, const FVector& RayDirection) override;
 	//~ End UGizmoElementBase Interface.
 
 	// Location of center of cylinder's base circle.
-	virtual void SetBase(const FVector& InBase);
-	virtual FVector GetBase() const;
+	INTERACTIVETOOLSFRAMEWORK_API virtual void SetBase(const FVector& InBase);
+	INTERACTIVETOOLSFRAMEWORK_API virtual FVector GetBase() const;
 
 	// Cylinder axis direction.
-	virtual void SetDirection(const FVector& InDirection);
-	virtual FVector GetDirection() const;
+	INTERACTIVETOOLSFRAMEWORK_API virtual void SetDirection(const FVector& InDirection);
+	INTERACTIVETOOLSFRAMEWORK_API virtual FVector GetDirection() const;
 
 	// Cylinder height.
-	virtual void SetHeight(float InHeight);
-	virtual float GetHeight() const;
+	INTERACTIVETOOLSFRAMEWORK_API virtual void SetHeight(float InHeight);
+	INTERACTIVETOOLSFRAMEWORK_API virtual float GetHeight() const;
 
 	// Cylinder radius.
-	virtual void SetRadius(float InRadius);
-	virtual float GetRadius() const;
+	INTERACTIVETOOLSFRAMEWORK_API virtual void SetRadius(float InRadius);
+	INTERACTIVETOOLSFRAMEWORK_API virtual float GetRadius() const;
 
 	// Number of sides for tessellating cylinder.
-	virtual void SetNumSides(int32 InNumSides);
-	virtual int32 GetNumSides() const;
+	INTERACTIVETOOLSFRAMEWORK_API virtual void SetNumSides(int32 InNumSides);
+	INTERACTIVETOOLSFRAMEWORK_API virtual int32 GetNumSides() const;
 
 protected:
 
diff --git a/Engine/Source/Runtime/InteractiveToolsFramework/Public/BaseGizmos/GizmoElementGroup.h b/Engine/Source/Runtime/InteractiveToolsFramework/Public/BaseGizmos/GizmoElementGroup.h
--- a/Engine/Source/Runtime/InteractiveToolsFramework/Public/BaseGizmos/GizmoElementGroup.h
+++ b/Engine/Source/Runtime/InteractiveToolsFramework/Public/BaseGizmos/GizmoElementGroup.h
@@ -12,49 +12,49 @@
  * Simple group object intended to be used as part of 3D Gizmos.
  * Contains multiple gizmo objects.
  */
-UCLASS(Transient)
-class INTERACTIVETOOLSFRAMEWORK_API UGizmoElementGroup : public UGizmoElementLineBase
+UCLASS(Transient, MinimalAPI)
+class UGizmoElementGroup : public UGizmoElementLineBase
 {
 	GENERATED_BODY()
 
 public:
 
 	//~ Begin UGizmoElementBase Interface.
-	virtual void Render(IToolsContextRenderAPI* RenderAPI, const FRenderTraversalState& RenderState) override;
-	virtual FInputRayHit LineTrace(const UGizmoViewContext* ViewContext, const FLineTraceTraversalState& LineTraceState, const FVector& RayOrigin, const FVector& RayDirection) override;
+	INTERACTIVETOOLSFRAMEWORK_API virtual void Render(IToolsContextRenderAPI* RenderAPI, const FRenderTraversalState& RenderState) override;
+	INTERACTIVETOOLSFRAMEWORK_API virtual FInputRayHit LineTrace(const UGizmoViewContext* ViewContext, const FLineTraceTraversalState& LineTraceState, const FVector& RayOrigin, const FVector& RayDirection) override;
 	//~ End UGizmoElementBase Interface.
 
 	// Add object to group.
-	virtual void Add(UGizmoElementBase* InElement);
+	INTERACTIVETOOLSFRAMEWORK_API virtual void Add(UGizmoElementBase* InElement);
 
 	// Remove object from group, if it exists
-	virtual void Remove(UGizmoElementBase* InElement);
+	INTERACTIVETOOLSFRAMEWORK_API virtual void Remove(UGizmoElementBase* InElement);
 
 	// Update group and contained elements' visibility state for elements in specified gizmo parts, return true if part was found.
-	virtual bool UpdatePartVisibleState(bool bVisible, uint32 InPartIdentifier);
+	INTERACTIVETOOLSFRAMEWORK_API virtual bool UpdatePartVisibleState(bool bVisible, uint32 InPartIdentifier);
 
 	// Get element's visible state for element associated with the specified gizmo part, if part was found.
-	virtual TOptional<bool> GetPartVisibleState(uint32 InPartIdentifier) const;
+	INTERACTIVETOOLSFRAMEWORK_API virtual TOptional<bool> GetPartVisibleState(uint32 InPartIdentifier) const;
 
 	// Update group and contained elements' hittable state for elements in specified gizmo parts, return true if part was found.
-	virtual bool UpdatePartHittableState(bool bHittable, uint32 InPartIdentifier);
+	INTERACTIVETOOLSFRAMEWORK_API virtual bool UpdatePartHittableState(bool bHittable, uint32 InPartIdentifier);
 
 	// Get element's hittable state for element associated with the specified gizmo part, if part was found.
-	virtual TOptional<bool> GetPartHittableState(uint32 InPartIdentifier) const;
+	INTERACTIVETOOLSFRAMEWORK_API virtual TOptional<bool> GetPartHittableState(uint32 InPartIdentifier) const;
 
 	// Update group and contained elements' interaction state for elements in specified gizmo parts, return true if part was found.
-	virtual bool UpdatePartInteractionState(EGizmoElementInteractionState InInteractionState, uint32 InPartIdentifier);
+	INTERACTIVETOOLSFRAMEWORK_API virtual bool UpdatePartInteractionState(EGizmoElementInteractionState InInteractionState, uint32 InPartIdentifier);
 
 	// Get element's interaction state for element associated with the specified gizmo part, if part was found.
-	virtual TOptional<EGizmoElementInteractionState> GetPartInteractionState(uint32 InPartIdentifiero) const;
+	INTERACTIVETOOLSFRAMEWORK_API virtual TOptional<EGizmoElementInteractionState> GetPartInteractionState(uint32 InPartIdentifiero) const;
 
 	// When true, maintains view-dependent constant scale for this gizmo object hierarchy
-	virtual void SetConstantScale(bool InConstantScale);
-	virtual bool GetConstantScale() const;
+	INTERACTIVETOOLSFRAMEWORK_API virtual void SetConstantScale(bool InConstantScale);
+	INTERACTIVETOOLSFRAMEWORK_API virtual bool GetConstantScale() const;
 
 	// Set whether this group should be treated as a hit owner and its part identifier returned when any of its sub-elements are hit.
-	virtual void SetHitOwner(bool bInHitOwner);
-	virtual bool GetHitOwner() const;
+	INTERACTIVETOOLSFRAMEWORK_API virtual void SetHitOwner(bool bInHitOwner);
+	INTERACTIVETOOLSFRAMEWORK_API virtual bool GetHitOwner() const;
 
 protected:
 
@@ -72,5 +72,5 @@ protected:
 	TArray<TObjectPtr<UGizmoElementBase>> Elements;
 
 	// Updates input transform's scale component to have uniform scale and applies constant scale if bConstantScale is true
-	virtual void ApplyUniformConstantScaleToTransform(double PixelToWorldScale, FTransform& InOutLocalToWorldTransform) const;
+	INTERACTIVETOOLSFRAMEWORK_API virtual void ApplyUniformConstantScaleToTransform(double PixelToWorldScale, FTransform& InOutLocalToWorldTransform) const;
 };
diff --git a/Engine/Source/Runtime/InteractiveToolsFramework/Public/BaseGizmos/GizmoElementHitTargets.h b/Engine/Source/Runtime/InteractiveToolsFramework/Public/BaseGizmos/GizmoElementHitTargets.h
--- a/Engine/Source/Runtime/InteractiveToolsFramework/Public/BaseGizmos/GizmoElementHitTargets.h
+++ b/Engine/Source/Runtime/InteractiveToolsFramework/Public/BaseGizmos/GizmoElementHitTargets.h
@@ -17,8 +17,8 @@ class UGizmoElementBase;
  * This hit target should be used for hitting a whole gizmo element hierarchy.
  * Use UGizmoElementHitMultiTarget, for hit targets that support hitting parts within a gizmo element hierarchy.
  */
-UCLASS()
-class INTERACTIVETOOLSFRAMEWORK_API UGizmoElementHitTarget : public UObject, public IGizmoClickTarget
+UCLASS(MinimalAPI)
+class UGizmoElementHitTarget : public UObject, public IGizmoClickTarget
 {
 	GENERATED_BODY()
 public:
@@ -44,14 +44,14 @@ public:
 	TFunction<bool(const FInputDeviceRay&)> Condition = nullptr;
 
 public:
-	virtual FInputRayHit IsHit(const FInputDeviceRay& ClickPos) const;
+	INTERACTIVETOOLSFRAMEWORK_API virtual FInputRayHit IsHit(const FInputDeviceRay& ClickPos) const;
 
-	virtual void UpdateHoverState(bool bHovering);
+	INTERACTIVETOOLSFRAMEWORK_API virtual void UpdateHoverState(bool bHovering);
 
-	virtual void UpdateInteractingState(bool bInteracting);
+	INTERACTIVETOOLSFRAMEWORK_API virtual void UpdateInteractingState(bool bInteracting);
 
 public:
-	static UGizmoElementHitTarget* Construct(
+	static INTERACTIVETOOLSFRAMEWORK_API UGizmoElementHitTarget* Construct(
 		UGizmoElementBase* InGizmoElement,
 		UGizmoViewContext* InGizmoViewContext,
 		UObject* Outer = (UObject*)GetTransientPackage());
@@ -67,8 +67,8 @@ public:
  * be defined in the gizmo. Identifier 0 is reserved for the default ID which should be assigned to 
  * elements that do not correspond to any gizmo part, such as non-hittable decorative elements.
  */
-UCLASS()
-class INTERACTIVETOOLSFRAMEWORK_API UGizmoElementHitMultiTarget : public UObject, public IGizmoClickMultiTarget
+UCLASS(MinimalAPI)
+class UGizmoElementHitMultiTarget : public UObject, public IGizmoClickMultiTarget
 {
 	GENERATED_BODY()
 public:
@@ -94,16 +94,16 @@ public:
 	TFunction<bool(const FInputDeviceRay&)> Condition = nullptr;
 
 public:
-	virtual FInputRayHit IsHit(const FInputDeviceRay& ClickPos) const;
+	INTERACTIVETOOLSFRAMEWORK_API virtual FInputRayHit IsHit(const FInputDeviceRay& ClickPos) const;
 
-	virtual void UpdateHoverState(bool bHovering, uint32 PartIdentifier);
+	INTERACTIVETOOLSFRAMEWORK_API virtual void UpdateHoverState(bool bHovering, uint32 PartIdentifier);
 
-	virtual void UpdateInteractingState(bool bInteracting, uint32 PartIdentifier);
+	INTERACTIVETOOLSFRAMEWORK_API virtual void UpdateInteractingState(bool bInteracting, uint32 PartIdentifier);
 
-	virtual void UpdateHittableState(bool bHittable, uint32 PartIdentifier);
+	INTERACTIVETOOLSFRAMEWORK_API virtual void UpdateHittableState(bool bHittable, uint32 PartIdentifier);
 
 public:
-	static UGizmoElementHitMultiTarget* Construct(
+	static INTERACTIVETOOLSFRAMEWORK_API UGizmoElementHitMultiTarget* Construct(
 		UGizmoElementBase* InGizmoElement,
 		UGizmoViewContext* InGizmoViewContext,
 		UObject* Outer = (UObject*)GetTransientPackage());
diff --git a/Engine/Source/Runtime/InteractiveToolsFramework/Public/BaseGizmos/GizmoElementLineBase.h b/Engine/Source/Runtime/InteractiveToolsFramework/Public/BaseGizmos/GizmoElementLineBase.h
--- a/Engine/Source/Runtime/InteractiveToolsFramework/Public/BaseGizmos/GizmoElementLineBase.h
+++ b/Engine/Source/Runtime/InteractiveToolsFramework/Public/BaseGizmos/GizmoElementLineBase.h
@@ -13,30 +13,30 @@
  * Base class for 2d and 3d primitive objects which support line drawing,
  * intended to be used as part of 3D Gizmos.
  */
-UCLASS(Transient, Abstract)
-class INTERACTIVETOOLSFRAMEWORK_API UGizmoElementLineBase : public UGizmoElementBase
+UCLASS(Transient, Abstract, MinimalAPI)
+class UGizmoElementLineBase : public UGizmoElementBase
 {
 	GENERATED_BODY()
 public:
 
 	// Get line thickness for based on current element interaction state and view
-	virtual float GetCurrentLineThickness(bool bPerspectiveView, float InViewFOV) const;
+	INTERACTIVETOOLSFRAMEWORK_API virtual float GetCurrentLineThickness(bool bPerspectiveView, float InViewFOV) const;
 
 	// Line thickness when rendering lines, 0.0 is valid and will render thinnest line 
-	virtual void SetLineThickness(float InLineThickness);
-	virtual float GetLineThickness() const;
+	INTERACTIVETOOLSFRAMEWORK_API virtual void SetLineThickness(float InLineThickness);
+	INTERACTIVETOOLSFRAMEWORK_API virtual float GetLineThickness() const;
 
 	// Multiplier applied to line thickness when hovering
-	virtual void SetHoverLineThicknessMultiplier(float InHoverLineThicknessMultiplier);
-	virtual float GetHoverLineThicknessMultiplier() const;
+	INTERACTIVETOOLSFRAMEWORK_API virtual void SetHoverLineThicknessMultiplier(float InHoverLineThicknessMultiplier);
+	INTERACTIVETOOLSFRAMEWORK_API virtual float GetHoverLineThicknessMultiplier() const;
 
 	// Multiplier applied to line thickness when interacting
-	virtual void SetInteractLineThicknessMultiplier(float InInteractLineThicknessMultiplier);
-	virtual float GetInteractLineThicknessMultiplier() const;
+	INTERACTIVETOOLSFRAMEWORK_API virtual void SetInteractLineThicknessMultiplier(float InInteractLineThicknessMultiplier);
+	INTERACTIVETOOLSFRAMEWORK_API virtual float GetInteractLineThicknessMultiplier() const;
 
 	// Whether line thickness is in screen space 
-	virtual void SetScreenSpaceLine(bool bInScreenSpaceLine);
-	virtual bool GetScreenSpaceLine() const;
+	INTERACTIVETOOLSFRAMEWORK_API virtual void SetScreenSpaceLine(bool bInScreenSpaceLine);
+	INTERACTIVETOOLSFRAMEWORK_API virtual bool GetScreenSpaceLine() const;
 
 	//
 	// Methods for managing line state attributes: LineColor, HoverLineColor, InteractLineColor
@@ -50,53 +50,53 @@ public:
 	// Set line render state line color attribute. 
 	//  @param InMaterial - line color to be set
 	//  @param InOverridesChildState - when true, this line color will override the line color of all child elements.
-	virtual void SetLineColor(FLinearColor InColor, bool InOverridesChildState = false);
+	INTERACTIVETOOLSFRAMEWORK_API virtual void SetLineColor(FLinearColor InColor, bool InOverridesChildState = false);
 
 	// Get line render state line color attribute's value.
-	virtual FLinearColor GetLineColor() const;
+	INTERACTIVETOOLSFRAMEWORK_API virtual FLinearColor GetLineColor() const;
 
 	// Returns true if line render state line color attribute has been set.
-	virtual bool HasLineColor() const;
+	INTERACTIVETOOLSFRAMEWORK_API virtual bool HasLineColor() const;
 
 	// Get line render state line color attribute's override setting. 
-	virtual bool DoesLineColorOverrideChildState() const;
+	INTERACTIVETOOLSFRAMEWORK_API virtual bool DoesLineColorOverrideChildState() const;
 
 	// Clear line render state line color attribute.
-	virtual void ClearLineColor();
+	INTERACTIVETOOLSFRAMEWORK_API virtual void ClearLineColor();
 
 	// Set line render state hover line color attribute. 
 	//  @param InMaterial - hover line color to be set
 	//  @param InOverridesChildState - when true, this hover line color will override the line color of all child elements.
-	virtual void SetHoverLineColor(FLinearColor InColor, bool InOverridesChildState = false);
+	INTERACTIVETOOLSFRAMEWORK_API virtual void SetHoverLineColor(FLinearColor InColor, bool InOverridesChildState = false);
 
 	// Get line render state hover line color attribute's value.
-	virtual FLinearColor GetHoverLineColor() const;
+	INTERACTIVETOOLSFRAMEWORK_API virtual FLinearColor GetHoverLineColor() const;
 
 	// Returns true if line render state hover line color attribute has been set. 
-	virtual bool HasHoverLineColor() const;
+	INTERACTIVETOOLSFRAMEWORK_API virtual bool HasHoverLineColor() const;
 
 	// Get line render state hover line color attribute's override setting.
-	virtual bool DoesHoverLineColorOverrideChildState() const;
+	INTERACTIVETOOLSFRAMEWORK_API virtual bool DoesHoverLineColorOverrideChildState() const;
 
 	// Clear line render state hover line color attribute. 
-	virtual void ClearHoverLineColor();
+	INTERACTIVETOOLSFRAMEWORK_API virtual void ClearHoverLineColor();
 
 	// Set line render state interact line color attribute. 
 	//  @param InMaterial - interact line color to be set
 	//  @param InOverridesChildState - when true, this interact line color will override the line color of all child elements.
-	virtual void SetInteractLineColor(FLinearColor InColor, bool InOverridesChildState = false);
+	INTERACTIVETOOLSFRAMEWORK_API virtual void SetInteractLineColor(FLinearColor InColor, bool InOverridesChildState = false);
 
 	// Get line render state interact line color attribute's value. 
-	virtual FLinearColor GetInteractLineColor() const;
+	INTERACTIVETOOLSFRAMEWORK_API virtual FLinearColor GetInteractLineColor() const;
 
 	// Returns true if line render state interact line color attribute has been set. 
-	virtual bool HasInteractLineColor() const;
+	INTERACTIVETOOLSFRAMEWORK_API virtual bool HasInteractLineColor() const;
 
 	// Get line render state interact line color attribute's override setting. 
-	virtual bool DoesInteractLineColorOverrideChildState() const;
+	INTERACTIVETOOLSFRAMEWORK_API virtual bool DoesInteractLineColorOverrideChildState() const;
 
 	// Clear line render state interact line color attribute.
-	virtual void ClearInteractLineColor();
+	INTERACTIVETOOLSFRAMEWORK_API virtual void ClearInteractLineColor();
 
 protected:
 
@@ -122,5 +122,5 @@ protected:
 
 	// Update render state during render traversal, determines the current render state for this element
 	// @return view dependent visibility, true if this element is visible in the current view.
-	virtual bool UpdateRenderState(IToolsContextRenderAPI* RenderAPI, const FVector& InLocalOrigin, FRenderTraversalState& InOutRenderState) override;
-};
\ No newline at end of file
+	INTERACTIVETOOLSFRAMEWORK_API virtual bool UpdateRenderState(IToolsContextRenderAPI* RenderAPI, const FVector& InLocalOrigin, FRenderTraversalState& InOutRenderState) override;
+};
diff --git a/Engine/Source/Runtime/InteractiveToolsFramework/Public/BaseGizmos/GizmoElementRectangle.h b/Engine/Source/Runtime/InteractiveToolsFramework/Public/BaseGizmos/GizmoElementRectangle.h
--- a/Engine/Source/Runtime/InteractiveToolsFramework/Public/BaseGizmos/GizmoElementRectangle.h
+++ b/Engine/Source/Runtime/InteractiveToolsFramework/Public/BaseGizmos/GizmoElementRectangle.h
@@ -15,53 +15,53 @@ class FMaterialRenderProxy;
  * Simple object intended to be used as part of 3D Gizmos.
  * Draws a rectangle based on parameters.
  */
-UCLASS(Transient)
-class INTERACTIVETOOLSFRAMEWORK_API UGizmoElementRectangle : public UGizmoElementLineBase
+UCLASS(Transient, MinimalAPI)
+class UGizmoElementRectangle : public UGizmoElementLineBase
 {
 	GENERATED_BODY()
 
 public:
 
 	//~ Begin UGizmoElementBase Interface.
-	virtual void Render(IToolsContextRenderAPI* RenderAPI, const FRenderTraversalState& RenderState) override;
-	virtual FInputRayHit LineTrace(const UGizmoViewContext* ViewContext, const FLineTraceTraversalState& LineTraceState, const FVector& RayOrigin, const FVector& RayDirection) override;
+	INTERACTIVETOOLSFRAMEWORK_API virtual void Render(IToolsContextRenderAPI* RenderAPI, const FRenderTraversalState& RenderState) override;
+	INTERACTIVETOOLSFRAMEWORK_API virtual FInputRayHit LineTrace(const UGizmoViewContext* ViewContext, const FLineTraceTraversalState& LineTraceState, const FVector& RayOrigin, const FVector& RayDirection) override;
 	//~ End UGizmoElementBase Interface.
 
 	// Location of rectangle center
-	virtual void SetCenter(FVector InCenter);
-	virtual FVector GetCenter() const;
+	INTERACTIVETOOLSFRAMEWORK_API virtual void SetCenter(FVector InCenter);
+	INTERACTIVETOOLSFRAMEWORK_API virtual FVector GetCenter() const;
 
 	// Width
-	virtual void SetWidth(float InWidth);
-	virtual float GetWidth() const;
+	INTERACTIVETOOLSFRAMEWORK_API virtual void SetWidth(float InWidth);
+	INTERACTIVETOOLSFRAMEWORK_API virtual float GetWidth() const;
 
 	// Height
-	virtual void SetHeight(float InHeight);
-	virtual float GetHeight() const;
+	INTERACTIVETOOLSFRAMEWORK_API virtual void SetHeight(float InHeight);
+	INTERACTIVETOOLSFRAMEWORK_API virtual float GetHeight() const;
 
 	// Up direction
-	virtual void SetUpDirection(const FVector& InUpDirection);
-	virtual FVector GetUpDirection() const;
+	INTERACTIVETOOLSFRAMEWORK_API virtual void SetUpDirection(const FVector& InUpDirection);
+	INTERACTIVETOOLSFRAMEWORK_API virtual FVector GetUpDirection() const;
 
 	// Side direction
-	virtual void SetSideDirection(const FVector& InSideDirection);
-	virtual FVector GetSideDirection() const;
+	INTERACTIVETOOLSFRAMEWORK_API virtual void SetSideDirection(const FVector& InSideDirection);
+	INTERACTIVETOOLSFRAMEWORK_API virtual FVector GetSideDirection() const;
 
 	// Draw mesh
-	virtual void SetDrawMesh(bool InDrawMesh);
-	virtual bool GetDrawMesh() const;
+	INTERACTIVETOOLSFRAMEWORK_API virtual void SetDrawMesh(bool InDrawMesh);
+	INTERACTIVETOOLSFRAMEWORK_API virtual bool GetDrawMesh() const;
 
 	// Draw line
-	virtual void SetDrawLine(bool InDrawLine);
-	virtual bool GetDrawLine() const;
+	INTERACTIVETOOLSFRAMEWORK_API virtual void SetDrawLine(bool InDrawLine);
+	INTERACTIVETOOLSFRAMEWORK_API virtual bool GetDrawLine() const;
 
 	// Hit mesh
-	virtual void SetHitMesh(bool InHitMesh);
-	virtual bool GetHitMesh() const;
+	INTERACTIVETOOLSFRAMEWORK_API virtual void SetHitMesh(bool InHitMesh);
+	INTERACTIVETOOLSFRAMEWORK_API virtual bool GetHitMesh() const;
 
 	// Hit line
-	virtual void SetHitLine(bool InHitLine);
-	virtual bool GetHitLine() const;
+	INTERACTIVETOOLSFRAMEWORK_API virtual void SetHitLine(bool InHitLine);
+	INTERACTIVETOOLSFRAMEWORK_API virtual bool GetHitLine() const;
 
 protected:
 
@@ -96,4 +96,4 @@ protected:
 
 	UPROPERTY()
 	bool bHitLine = false;
-};
\ No newline at end of file
+};
diff --git a/Engine/Source/Runtime/InteractiveToolsFramework/Public/BaseGizmos/GizmoElementRenderState.h b/Engine/Source/Runtime/InteractiveToolsFramework/Public/BaseGizmos/GizmoElementRenderState.h
--- a/Engine/Source/Runtime/InteractiveToolsFramework/Public/BaseGizmos/GizmoElementRenderState.h
+++ b/Engine/Source/Runtime/InteractiveToolsFramework/Public/BaseGizmos/GizmoElementRenderState.h
@@ -20,7 +20,7 @@
 // Stores color state value, whether state is set and whether it overrides child state. 
 //
 USTRUCT()
-struct INTERACTIVETOOLSFRAMEWORK_API FGizmoElementColorAttribute
+struct FGizmoElementColorAttribute
 {
 	GENERATED_BODY()
 
@@ -39,16 +39,16 @@ struct INTERACTIVETOOLSFRAMEWORK_API FGizmoElementColorAttribute
 	bool bOverridesChildState = false;
 
 	// Return the color value. This returns the default color value if no value has been set.
-	FLinearColor GetColor() const;
+	INTERACTIVETOOLSFRAMEWORK_API FLinearColor GetColor() const;
 
 	// Set color value and override flag
-	void SetColor(FLinearColor InColor, bool InOverridesChildState = false);
+	INTERACTIVETOOLSFRAMEWORK_API void SetColor(FLinearColor InColor, bool InOverridesChildState = false);
 
 	// Reset attribute to default values.
-	void Reset();
+	INTERACTIVETOOLSFRAMEWORK_API void Reset();
 
 	// Update this state based on a child state attribute.
-	void UpdateState(const FGizmoElementColorAttribute& InChildColorAttribute);
+	INTERACTIVETOOLSFRAMEWORK_API void UpdateState(const FGizmoElementColorAttribute& InChildColorAttribute);
 };
 
 //
@@ -57,7 +57,7 @@ struct INTERACTIVETOOLSFRAMEWORK_API FGizmoElementColorAttribute
 // Stores material state value, whether state is set and whether it overrides child state.
 //
 USTRUCT()
-struct INTERACTIVETOOLSFRAMEWORK_API FGizmoElementMaterialAttribute 
+struct FGizmoElementMaterialAttribute 
 {
 	GENERATED_BODY()
 
@@ -70,23 +70,23 @@ struct INTERACTIVETOOLSFRAMEWORK_API FGizmoElementMaterialAttribute
 	bool bOverridesChildState = false;
 
 	// Return the material value.
-	const UMaterialInterface* GetMaterial() const;
+	INTERACTIVETOOLSFRAMEWORK_API const UMaterialInterface* GetMaterial() const;
 
 	// Set material value and override flag
-	void SetMaterial(TWeakObjectPtr<UMaterialInterface> InColor, bool InOverridesChildState = false);
+	INTERACTIVETOOLSFRAMEWORK_API void SetMaterial(TWeakObjectPtr<UMaterialInterface> InColor, bool InOverridesChildState = false);
 
 	// Reset attribute to default values.
-	void Reset();
+	INTERACTIVETOOLSFRAMEWORK_API void Reset();
 
 	// Update this state based on a child state attribute.
-	void UpdateState(const FGizmoElementMaterialAttribute& InChildMaterialAttribute);
+	INTERACTIVETOOLSFRAMEWORK_API void UpdateState(const FGizmoElementMaterialAttribute& InChildMaterialAttribute);
 };
 
 //
 // Mesh render state structure.
 //
 USTRUCT()
-struct INTERACTIVETOOLSFRAMEWORK_API FGizmoElementMeshRenderStateAttributes
+struct FGizmoElementMeshRenderStateAttributes
 {
 	GENERATED_BODY()
 
@@ -115,20 +115,20 @@ struct INTERACTIVETOOLSFRAMEWORK_API FGizmoElementMeshRenderStateAttributes
 	FGizmoElementColorAttribute InteractVertexColor;
 
 	// Returns the material corresponding to the input interaction state
-	const UMaterialInterface* GetMaterial(EGizmoElementInteractionState InteractionState);
+	INTERACTIVETOOLSFRAMEWORK_API const UMaterialInterface* GetMaterial(EGizmoElementInteractionState InteractionState);
 
 	// Returns mesh vertex color
-	FLinearColor GetVertexColor(EGizmoElementInteractionState InteractionState);
+	INTERACTIVETOOLSFRAMEWORK_API FLinearColor GetVertexColor(EGizmoElementInteractionState InteractionState);
 
 	// Update this mesh render state based on a child mesh render state attribute.
-	void Update(FGizmoElementMeshRenderStateAttributes& InChildAttributes);
+	INTERACTIVETOOLSFRAMEWORK_API void Update(FGizmoElementMeshRenderStateAttributes& InChildAttributes);
 };
 
 //
 // Line render state structure.
 //
 USTRUCT()
-struct INTERACTIVETOOLSFRAMEWORK_API FGizmoElementLineRenderStateAttributes
+struct FGizmoElementLineRenderStateAttributes
 {
 	GENERATED_BODY()
 
@@ -145,9 +145,9 @@ struct INTERACTIVETOOLSFRAMEWORK_API FGizmoElementLineRenderStateAttributes
 	FGizmoElementColorAttribute InteractLineColor;
 
 	// Returns the line color corresponding to the input interaction state
-	FLinearColor GetLineColor(EGizmoElementInteractionState InteractionState);
+	INTERACTIVETOOLSFRAMEWORK_API FLinearColor GetLineColor(EGizmoElementInteractionState InteractionState);
 
 	// Update this line render state based on a child mesh render state attribute.
-	void Update(FGizmoElementLineRenderStateAttributes& InChildAttributes);
+	INTERACTIVETOOLSFRAMEWORK_API void Update(FGizmoElementLineRenderStateAttributes& InChildAttributes);
 };
 
diff --git a/Engine/Source/Runtime/InteractiveToolsFramework/Public/BaseGizmos/GizmoElementTorus.h b/Engine/Source/Runtime/InteractiveToolsFramework/Public/BaseGizmos/GizmoElementTorus.h
--- a/Engine/Source/Runtime/InteractiveToolsFramework/Public/BaseGizmos/GizmoElementTorus.h
+++ b/Engine/Source/Runtime/InteractiveToolsFramework/Public/BaseGizmos/GizmoElementTorus.h
@@ -18,34 +18,34 @@ class FMaterialRenderProxy;
  * Note: the LineTrace method does not perform a true ray-torus intersection!
  * See comment above LineTrace method below for details of how this intersection is approximated.
  */
-UCLASS(Transient)
-class INTERACTIVETOOLSFRAMEWORK_API UGizmoElementTorus : public UGizmoElementCircleBase
+UCLASS(Transient, MinimalAPI)
+class UGizmoElementTorus : public UGizmoElementCircleBase
 {
 	GENERATED_BODY()
 
 public:
 
 	//~ Begin UGizmoElementBase Interface.
-	virtual void Render(IToolsContextRenderAPI* RenderAPI, const FRenderTraversalState& RenderState) override;
+	INTERACTIVETOOLSFRAMEWORK_API virtual void Render(IToolsContextRenderAPI* RenderAPI, const FRenderTraversalState& RenderState) override;
 
 	// LineTrace approximates ray-torus intersection by intersecting the ray with the plane in which the torus 
 	// lies, then determining a hit point closest to the linear circle defined by torus center and torus outer radius.
 	// If the torus lies at a glancing angle, the ray-torus intersection is performed against cylinders approximating
 	// the shape of the torus.
-	virtual FInputRayHit LineTrace(const UGizmoViewContext* ViewContext, const FLineTraceTraversalState& LineTraceState, const FVector& RayOrigin, const FVector& RayDirection) override;
+	INTERACTIVETOOLSFRAMEWORK_API virtual FInputRayHit LineTrace(const UGizmoViewContext* ViewContext, const FLineTraceTraversalState& LineTraceState, const FVector& RayOrigin, const FVector& RayDirection) override;
 	//~ End UGizmoElementBase Interface.
 
 	// Inner circles radius.
-	virtual void SetInnerRadius(double InInnerRadius);
-	virtual double GetInnerRadius() const;
+	INTERACTIVETOOLSFRAMEWORK_API virtual void SetInnerRadius(double InInnerRadius);
+	INTERACTIVETOOLSFRAMEWORK_API virtual double GetInnerRadius() const;
 
 	// Number of inner slices for rendering torus.
-	virtual void SetNumInnerSlices(int32 InNumInnerSlices);
-	virtual int32 GetNumInnerSlices() const;
+	INTERACTIVETOOLSFRAMEWORK_API virtual void SetNumInnerSlices(int32 InNumInnerSlices);
+	INTERACTIVETOOLSFRAMEWORK_API virtual int32 GetNumInnerSlices() const;
 
 	// If partial, renders end caps when true.
-	virtual void SetEndCaps(bool InEndCaps);
-	virtual bool GetEndCaps() const;
+	INTERACTIVETOOLSFRAMEWORK_API virtual void SetEndCaps(bool InEndCaps);
+	INTERACTIVETOOLSFRAMEWORK_API virtual bool GetEndCaps() const;
 
 protected:
 
@@ -63,6 +63,6 @@ protected:
 
 protected:
 
-	bool IsGlancingAngle(const UGizmoViewContext* View, const FVector& InWorldCenter, const FVector& InWorldNormal);
+	INTERACTIVETOOLSFRAMEWORK_API bool IsGlancingAngle(const UGizmoViewContext* View, const FVector& InWorldCenter, const FVector& InWorldNormal);
 
-};
\ No newline at end of file
+};
diff --git a/Engine/Source/Runtime/InteractiveToolsFramework/Public/BaseGizmos/GizmoInterfaces.h b/Engine/Source/Runtime/InteractiveToolsFramework/Public/BaseGizmos/GizmoInterfaces.h
--- a/Engine/Source/Runtime/InteractiveToolsFramework/Public/BaseGizmos/GizmoInterfaces.h
+++ b/Engine/Source/Runtime/InteractiveToolsFramework/Public/BaseGizmos/GizmoInterfaces.h
@@ -25,15 +25,15 @@ struct FFrame;
 
 
 
-UINTERFACE()
-class INTERACTIVETOOLSFRAMEWORK_API UGizmoTransformSource : public UInterface
+UINTERFACE(MinimalAPI)
+class UGizmoTransformSource : public UInterface
 {
 	GENERATED_BODY()
 };
 /**
  * IGizmoTransformSource is an interface which is used to Get/Set an FTransform.
  */
-class INTERACTIVETOOLSFRAMEWORK_API IGizmoTransformSource
+class IGizmoTransformSource
 {
 	GENERATED_BODY()
 public:
@@ -48,8 +48,8 @@ public:
 
 
 
-UINTERFACE()
-class INTERACTIVETOOLSFRAMEWORK_API UGizmoAxisSource : public UInterface
+UINTERFACE(MinimalAPI)
+class UGizmoAxisSource : public UInterface
 {
 	GENERATED_BODY()
 };
@@ -60,7 +60,7 @@ class INTERACTIVETOOLSFRAMEWORK_API UGizmoAxisSource : public UInterface
  * assumed to be mutually-orthogonal and perpendicular to the Axis Direction
  * (ie that's the normal and the 3 vectors form a coordinate frame).
  */
-class INTERACTIVETOOLSFRAMEWORK_API IGizmoAxisSource
+class IGizmoAxisSource
 {
 	GENERATED_BODY()
 public:
@@ -87,7 +87,7 @@ public:
 	 * Utility function that always returns a 3D coordinate system (ie plane normal and perpendicular axes).
 	 * Internally calls GetTangentVectors() if available, otherwise constructs arbitrary mutually perpendicular vectors. 
 	 */
-	void GetAxisFrame(
+	INTERACTIVETOOLSFRAMEWORK_API void GetAxisFrame(
 		FVector& PlaneNormalOut, FVector& PlaneAxis1Out, FVector& PlaneAxis2Out) const;
 };
 
@@ -96,15 +96,15 @@ public:
 
 
 
-UINTERFACE()
-class INTERACTIVETOOLSFRAMEWORK_API UGizmoClickTarget : public UInterface
+UINTERFACE(MinimalAPI)
+class UGizmoClickTarget : public UInterface
 {
 	GENERATED_BODY()
 };
 /**
  * IGizmoClickTarget is an interface used to provide a ray-object hit test.
  */
-class INTERACTIVETOOLSFRAMEWORK_API IGizmoClickTarget
+class IGizmoClickTarget
 {
 	GENERATED_BODY()
 public:
@@ -129,8 +129,8 @@ public:
 	virtual void UpdateInteractingState(bool bInteracting) = 0;
 };
 
-UINTERFACE()
-class INTERACTIVETOOLSFRAMEWORK_API UGizmoClickMultiTarget : public UInterface
+UINTERFACE(MinimalAPI)
+class UGizmoClickMultiTarget : public UInterface
 {
 	GENERATED_BODY()
 };
@@ -143,7 +143,7 @@ class INTERACTIVETOOLSFRAMEWORK_API UGizmoClickMultiTarget : public UInterface
  * be defined in the gizmo. Identifier 0 is reserved for the default ID which should be assigned to 
  * elements that do not correspond to any gizmo part, such as non-hittable decorative elements.
  */
-class INTERACTIVETOOLSFRAMEWORK_API IGizmoClickMultiTarget
+class IGizmoClickMultiTarget
 {
 	GENERATED_BODY()
 public:
@@ -179,8 +179,8 @@ public:
 
 
 
-UINTERFACE()
-class INTERACTIVETOOLSFRAMEWORK_API UGizmoRenderTarget : public UInterface
+UINTERFACE(MinimalAPI)
+class UGizmoRenderTarget : public UInterface
 {
 	GENERATED_BODY()
 };
@@ -188,7 +188,7 @@ class INTERACTIVETOOLSFRAMEWORK_API UGizmoRenderTarget : public UInterface
 /**
  * UGizmoRenderTarget is an interface used to provide rendering of a target
  */
-class INTERACTIVETOOLSFRAMEWORK_API IGizmoRenderTarget
+class IGizmoRenderTarget
 {
 	GENERATED_BODY()
 public:
@@ -201,8 +201,8 @@ public:
 
 
 
-UINTERFACE()
-class INTERACTIVETOOLSFRAMEWORK_API UGizmoRenderMultiTarget : public UInterface
+UINTERFACE(MinimalAPI)
+class UGizmoRenderMultiTarget : public UInterface
 {
 	GENERATED_BODY()
 };
@@ -216,7 +216,7 @@ class INTERACTIVETOOLSFRAMEWORK_API UGizmoRenderMultiTarget : public UInterface
  * be defined in the gizmo. Identifier 0 is reserved for the default ID which should be assigned to 
  * elements that do not correspond to any gizmo part, such as non-hittable decorative elements.
  */
-class INTERACTIVETOOLSFRAMEWORK_API IGizmoRenderMultiTarget
+class IGizmoRenderMultiTarget
 {
 	GENERATED_BODY()
 public:
@@ -233,15 +233,15 @@ public:
 };
 
 
-UINTERFACE()
-class INTERACTIVETOOLSFRAMEWORK_API UGizmoStateTarget : public UInterface
+UINTERFACE(MinimalAPI)
+class UGizmoStateTarget : public UInterface
 {
 	GENERATED_BODY()
 };
 /**
  * IGizmoStateTarget is an interface that is used to pass notifications about significant gizmo state updates
  */
-class INTERACTIVETOOLSFRAMEWORK_API IGizmoStateTarget
+class IGizmoStateTarget
 {
 	GENERATED_BODY()
 public:
@@ -262,15 +262,15 @@ public:
 
 
 
-UINTERFACE()
-class INTERACTIVETOOLSFRAMEWORK_API UGizmoFloatParameterSource : public UInterface
+UINTERFACE(MinimalAPI)
+class UGizmoFloatParameterSource : public UInterface
 {
 	GENERATED_BODY()
 };
 /**
  * IGizmoFloatParameterSource provides Get and Set for an arbitrary float-valued parameter.
  */
-class INTERACTIVETOOLSFRAMEWORK_API IGizmoFloatParameterSource
+class IGizmoFloatParameterSource
 {
 	GENERATED_BODY()
 public:
@@ -295,15 +295,15 @@ public:
 
 
 
-UINTERFACE()
-class INTERACTIVETOOLSFRAMEWORK_API UGizmoVec2ParameterSource : public UInterface
+UINTERFACE(MinimalAPI)
+class UGizmoVec2ParameterSource : public UInterface
 {
 	GENERATED_BODY()
 };
 /**
  * IGizmoVec2ParameterSource provides Get and Set for an arbitrary 2D-vector-valued parameter.
  */
-class INTERACTIVETOOLSFRAMEWORK_API IGizmoVec2ParameterSource
+class IGizmoVec2ParameterSource
 {
 	GENERATED_BODY()
 public:
diff --git a/Engine/Source/Runtime/InteractiveToolsFramework/Public/BaseGizmos/GizmoLineHandleComponent.h b/Engine/Source/Runtime/InteractiveToolsFramework/Public/BaseGizmos/GizmoLineHandleComponent.h
--- a/Engine/Source/Runtime/InteractiveToolsFramework/Public/BaseGizmos/GizmoLineHandleComponent.h
+++ b/Engine/Source/Runtime/InteractiveToolsFramework/Public/BaseGizmos/GizmoLineHandleComponent.h
@@ -12,8 +12,8 @@
  * Simple Component intended to be used as part of 3D Gizmos.
  * Draws line terminated by a short perpendicular handle line based on parameters.
  */
-UCLASS(ClassGroup = Utility, HideCategories = (Physics, Collision, Lighting, Rendering, Mobile))
-class INTERACTIVETOOLSFRAMEWORK_API UGizmoLineHandleComponent : public UGizmoBaseComponent
+UCLASS(ClassGroup = Utility, HideCategories = (Physics, Collision, Lighting, Rendering, Mobile), MinimalAPI)
+class UGizmoLineHandleComponent : public UGizmoBaseComponent
 {
 	GENERATED_BODY()
 
@@ -37,11 +37,11 @@ public:
 	bool bImageScale = true;
 private:
 	//~ Begin UPrimitiveComponent Interface.
-	virtual FPrimitiveSceneProxy* CreateSceneProxy() override;
-	virtual bool LineTraceComponent(FHitResult& OutHit, const FVector Start, const FVector End, const FCollisionQueryParams& Params) override;
+	INTERACTIVETOOLSFRAMEWORK_API virtual FPrimitiveSceneProxy* CreateSceneProxy() override;
+	INTERACTIVETOOLSFRAMEWORK_API virtual bool LineTraceComponent(FHitResult& OutHit, const FVector Start, const FVector End, const FCollisionQueryParams& Params) override;
 	//~ End UPrimitiveComponent Interface.
 
 	//~ Begin USceneComponent Interface.
-	virtual FBoxSphereBounds CalcBounds(const FTransform& LocalToWorld) const override;
+	INTERACTIVETOOLSFRAMEWORK_API virtual FBoxSphereBounds CalcBounds(const FTransform& LocalToWorld) const override;
 	//~ Begin USceneComponent Interface.
-};
\ No newline at end of file
+};
diff --git a/Engine/Source/Runtime/InteractiveToolsFramework/Public/BaseGizmos/GizmoRectangleComponent.h b/Engine/Source/Runtime/InteractiveToolsFramework/Public/BaseGizmos/GizmoRectangleComponent.h
--- a/Engine/Source/Runtime/InteractiveToolsFramework/Public/BaseGizmos/GizmoRectangleComponent.h
+++ b/Engine/Source/Runtime/InteractiveToolsFramework/Public/BaseGizmos/GizmoRectangleComponent.h
@@ -11,8 +11,8 @@
  * Simple Component intended to be used as part of 3D Gizmos. 
  * Draws outline of 3D rectangle based on parameters.
  */
-UCLASS(ClassGroup = Utility, HideCategories = (Physics, Collision, Lighting, Rendering, Mobile))
-class INTERACTIVETOOLSFRAMEWORK_API UGizmoRectangleComponent : public UGizmoBaseComponent
+UCLASS(ClassGroup = Utility, HideCategories = (Physics, Collision, Lighting, Rendering, Mobile), MinimalAPI)
+class UGizmoRectangleComponent : public UGizmoBaseComponent
 {
 	GENERATED_BODY()
 
@@ -50,12 +50,12 @@ public:
 
 private:
 	//~ Begin UPrimitiveComponent Interface.
-	virtual FPrimitiveSceneProxy* CreateSceneProxy() override;
-	virtual bool LineTraceComponent(FHitResult& OutHit, const FVector Start, const FVector End, const FCollisionQueryParams& Params) override;
-	virtual void GetUsedMaterials(TArray<UMaterialInterface*>& OutMaterials, bool bGetDebugMaterials = false) const override;
+	INTERACTIVETOOLSFRAMEWORK_API virtual FPrimitiveSceneProxy* CreateSceneProxy() override;
+	INTERACTIVETOOLSFRAMEWORK_API virtual bool LineTraceComponent(FHitResult& OutHit, const FVector Start, const FVector End, const FCollisionQueryParams& Params) override;
+	INTERACTIVETOOLSFRAMEWORK_API virtual void GetUsedMaterials(TArray<UMaterialInterface*>& OutMaterials, bool bGetDebugMaterials = false) const override;
 	//~ End UPrimitiveComponent Interface.
 
 	//~ Begin USceneComponent Interface.
-	virtual FBoxSphereBounds CalcBounds(const FTransform& LocalToWorld) const override;
+	INTERACTIVETOOLSFRAMEWORK_API virtual FBoxSphereBounds CalcBounds(const FTransform& LocalToWorld) const override;
 	//~ Begin USceneComponent Interface.
-};
\ No newline at end of file
+};
diff --git a/Engine/Source/Runtime/InteractiveToolsFramework/Public/BaseGizmos/GizmoViewContext.h b/Engine/Source/Runtime/InteractiveToolsFramework/Public/BaseGizmos/GizmoViewContext.h
--- a/Engine/Source/Runtime/InteractiveToolsFramework/Public/BaseGizmos/GizmoViewContext.h
+++ b/Engine/Source/Runtime/InteractiveToolsFramework/Public/BaseGizmos/GizmoViewContext.h
@@ -14,8 +14,8 @@
  * FSceneView so that functions could be written in a templated way to use either FSceneView
  * or UGizmoViewContext, though UGizmoViewContext only keeps the needed data.
  */
-UCLASS()
-class INTERACTIVETOOLSFRAMEWORK_API UGizmoViewContext : public UObject
+UCLASS(MinimalAPI)
+class UGizmoViewContext : public UObject
 {
 	GENERATED_BODY()
 public:
@@ -91,4 +91,4 @@ public:
 
 protected:
 	bool bIsPerspectiveProjection;
-};
\ No newline at end of file
+};
diff --git a/Engine/Source/Runtime/InteractiveToolsFramework/Public/BaseGizmos/HitTargets.h b/Engine/Source/Runtime/InteractiveToolsFramework/Public/BaseGizmos/HitTargets.h
--- a/Engine/Source/Runtime/InteractiveToolsFramework/Public/BaseGizmos/HitTargets.h
+++ b/Engine/Source/Runtime/InteractiveToolsFramework/Public/BaseGizmos/HitTargets.h
@@ -13,8 +13,8 @@ class UPrimitiveComponent;
  * UGizmoLambdaHitTarget is an IGizmoClickTarget implementation that
  * forwards the hit-test function to a TFunction
  */
-UCLASS()
-class INTERACTIVETOOLSFRAMEWORK_API UGizmoLambdaHitTarget : public UObject, public IGizmoClickTarget
+UCLASS(MinimalAPI)
+class UGizmoLambdaHitTarget : public UObject, public IGizmoClickTarget
 {
 	GENERATED_BODY()
 public:
@@ -29,11 +29,11 @@ public:
 
 
 public:
-	virtual FInputRayHit IsHit(const FInputDeviceRay& ClickPos) const;
+	INTERACTIVETOOLSFRAMEWORK_API virtual FInputRayHit IsHit(const FInputDeviceRay& ClickPos) const;
 
-	virtual void UpdateHoverState(bool bHovering);
+	INTERACTIVETOOLSFRAMEWORK_API virtual void UpdateHoverState(bool bHovering);
 
-	virtual void UpdateInteractingState(bool bHovering);
+	INTERACTIVETOOLSFRAMEWORK_API virtual void UpdateInteractingState(bool bHovering);
 };
 
 
@@ -42,8 +42,8 @@ public:
  * UGizmoComponentHitTarget is an IGizmoClickTarget implementation that
  * hit-tests a UPrimitiveComponent
  */
-UCLASS()
-class INTERACTIVETOOLSFRAMEWORK_API UGizmoComponentHitTarget : public UObject, public IGizmoClickTarget
+UCLASS(MinimalAPI)
+class UGizmoComponentHitTarget : public UObject, public IGizmoClickTarget
 {
 	GENERATED_BODY()
 public:
@@ -69,14 +69,14 @@ public:
 	TFunction<void(bool)> UpdateInteractingFunction;
 
 public:
-	virtual FInputRayHit IsHit(const FInputDeviceRay& ClickPos) const;
+	INTERACTIVETOOLSFRAMEWORK_API virtual FInputRayHit IsHit(const FInputDeviceRay& ClickPos) const;
 
-	virtual void UpdateHoverState(bool bHovering);
+	INTERACTIVETOOLSFRAMEWORK_API virtual void UpdateHoverState(bool bHovering);
 
-	virtual void UpdateInteractingState(bool bHovering);
+	INTERACTIVETOOLSFRAMEWORK_API virtual void UpdateInteractingState(bool bHovering);
 
 public:
-	static UGizmoComponentHitTarget* Construct(
+	static INTERACTIVETOOLSFRAMEWORK_API UGizmoComponentHitTarget* Construct(
 		UPrimitiveComponent* Component,
 		UObject* Outer = (UObject*)GetTransientPackage());
 };
diff --git a/Engine/Source/Runtime/InteractiveToolsFramework/Public/BaseGizmos/IntervalGizmo.h b/Engine/Source/Runtime/InteractiveToolsFramework/Public/BaseGizmos/IntervalGizmo.h
--- a/Engine/Source/Runtime/InteractiveToolsFramework/Public/BaseGizmos/IntervalGizmo.h
+++ b/Engine/Source/Runtime/InteractiveToolsFramework/Public/BaseGizmos/IntervalGizmo.h
@@ -31,13 +31,13 @@ class UGizmoViewContext;
  * If a particular sub-Gizmo is not required, simply set that UProperty to null.
 
  */
-UCLASS(Transient, NotPlaceable, Hidden, NotBlueprintable, NotBlueprintType)
-class INTERACTIVETOOLSFRAMEWORK_API AIntervalGizmoActor : public AGizmoActor
+UCLASS(Transient, NotPlaceable, Hidden, NotBlueprintable, NotBlueprintType, MinimalAPI)
+class AIntervalGizmoActor : public AGizmoActor
 {
 	GENERATED_BODY()
 public:
 
-	AIntervalGizmoActor();
+	INTERACTIVETOOLSFRAMEWORK_API AIntervalGizmoActor();
 	
 public:
 
@@ -56,7 +56,7 @@ public:
 	 * Create a new instance of AIntervalGizmoActor and populate the various
 	 * sub-components with standard GizmoXComponent instances suitable for a 3-interval Gizmo
 	 */
-	static AIntervalGizmoActor* ConstructDefaultIntervalGizmo(UWorld* World, 
+	static INTERACTIVETOOLSFRAMEWORK_API AIntervalGizmoActor* ConstructDefaultIntervalGizmo(UWorld* World, 
 		UGizmoViewContext* GizmoViewContext);
 
 };
@@ -73,7 +73,7 @@ public:
  * a three-axis Interval Gizmo, override this function to customize
  * the Actor sub-elements.
  */
-class INTERACTIVETOOLSFRAMEWORK_API FIntervalGizmoActorFactory
+class FIntervalGizmoActorFactory
 {
 public:
 	FIntervalGizmoActorFactory(UGizmoViewContext* GizmoViewContextIn)
@@ -100,8 +100,8 @@ protected:
 
 
 
-UCLASS()
-class INTERACTIVETOOLSFRAMEWORK_API UIntervalGizmoBuilder : public UInteractiveGizmoBuilder
+UCLASS(MinimalAPI)
+class UIntervalGizmoBuilder : public UInteractiveGizmoBuilder
 {
 	GENERATED_BODY()
 
@@ -133,35 +133,35 @@ public:
 	TFunction<void(UPrimitiveComponent*, EToolContextCoordinateSystem)> UpdateCoordSystemFunction;
 
 
-	virtual UInteractiveGizmo* BuildGizmo(const FToolBuilderState& SceneState) const override;
+	INTERACTIVETOOLSFRAMEWORK_API virtual UInteractiveGizmo* BuildGizmo(const FToolBuilderState& SceneState) const override;
 };
 
 
-UCLASS()
-class INTERACTIVETOOLSFRAMEWORK_API UIntervalGizmo : public UInteractiveGizmo
+UCLASS(MinimalAPI)
+class UIntervalGizmo : public UInteractiveGizmo
 {
 	GENERATED_BODY()
 
 public:
 
-	static FString GizmoName;
+	static INTERACTIVETOOLSFRAMEWORK_API FString GizmoName;
 
-	virtual void SetWorld(UWorld* WorldIn);
-	virtual void SetGizmoActorBuilder(TSharedPtr<FIntervalGizmoActorFactory> Builder);
-	virtual void SetUpdateHoverFunction(TFunction<void(UPrimitiveComponent*, bool)> HoverFunction);
-	virtual void SetUpdateCoordSystemFunction(TFunction<void(UPrimitiveComponent*, EToolContextCoordinateSystem)> CoordSysFunction);
+	INTERACTIVETOOLSFRAMEWORK_API virtual void SetWorld(UWorld* WorldIn);
+	INTERACTIVETOOLSFRAMEWORK_API virtual void SetGizmoActorBuilder(TSharedPtr<FIntervalGizmoActorFactory> Builder);
+	INTERACTIVETOOLSFRAMEWORK_API virtual void SetUpdateHoverFunction(TFunction<void(UPrimitiveComponent*, bool)> HoverFunction);
+	INTERACTIVETOOLSFRAMEWORK_API virtual void SetUpdateCoordSystemFunction(TFunction<void(UPrimitiveComponent*, EToolContextCoordinateSystem)> CoordSysFunction);
 
 	// UInteractiveGizmo overrides
-	virtual void Setup() override;
-	virtual void Shutdown() override;
-	virtual void Tick(float DeltaTime) override;
+	INTERACTIVETOOLSFRAMEWORK_API virtual void Setup() override;
+	INTERACTIVETOOLSFRAMEWORK_API virtual void Shutdown() override;
+	INTERACTIVETOOLSFRAMEWORK_API virtual void Tick(float DeltaTime) override;
 
 
-	virtual void SetActiveTarget(UTransformProxy* TransformTargetIn, UGizmoLocalFloatParameterSource* UpInterval, UGizmoLocalFloatParameterSource* DownInterval, UGizmoLocalFloatParameterSource* ForwardInterval,
+	INTERACTIVETOOLSFRAMEWORK_API virtual void SetActiveTarget(UTransformProxy* TransformTargetIn, UGizmoLocalFloatParameterSource* UpInterval, UGizmoLocalFloatParameterSource* DownInterval, UGizmoLocalFloatParameterSource* ForwardInterval,
 		                         IToolContextTransactionProvider* TransactionProvider = nullptr);
 
 	/** Sets functions that allow the endpoints of the intervals to be snapped to world geometry when ShouldAlignDestination is true */
-	virtual void SetWorldAlignmentFunctions(
+	INTERACTIVETOOLSFRAMEWORK_API virtual void SetWorldAlignmentFunctions(
 		TUniqueFunction<bool()>&& ShouldAlignDestination,
 		TUniqueFunction<bool(const FRay&, FVector&)>&& DestinationAlignmentRayCaster
 	);
@@ -173,17 +173,17 @@ public:
 	/**
 	* Clear the parameter sources for this gizmo
 	*/
-	virtual void ClearSources();
+	INTERACTIVETOOLSFRAMEWORK_API virtual void ClearSources();
 
 	/**
 	 * Clear the active target object for the Gizmo
 	 */
-	virtual void ClearActiveTarget();
+	INTERACTIVETOOLSFRAMEWORK_API virtual void ClearActiveTarget();
 
 	/**
 	 * Gets the location and orientation of the interval gizmo.
 	 */
-	virtual FTransform GetGizmoTransform() const;
+	INTERACTIVETOOLSFRAMEWORK_API virtual FTransform GetGizmoTransform() const;
 
 	/** State target is shared across gizmos, and created internally during SetActiveTarget() */
 	UPROPERTY()
@@ -253,7 +253,7 @@ protected:
 protected:
 
 	/** @return a new instance of the standard axis-handle Gizmo */
-	virtual UInteractiveGizmo* AddIntervalHandleGizmo(
+	INTERACTIVETOOLSFRAMEWORK_API virtual UInteractiveGizmo* AddIntervalHandleGizmo(
 		USceneComponent* RootComponent,
 		UPrimitiveComponent* HandleComponent, 
 		IGizmoAxisSource* AxisSource,
@@ -271,18 +271,18 @@ protected:
  *
  * This ParameterSource is intended to be used to create 3D Axis Interval Gizmos.
  */
-UCLASS()
-class INTERACTIVETOOLSFRAMEWORK_API UGizmoAxisIntervalParameterSource : public UGizmoBaseFloatParameterSource
+UCLASS(MinimalAPI)
+class UGizmoAxisIntervalParameterSource : public UGizmoBaseFloatParameterSource
 {
 	GENERATED_BODY()
 public:
-	virtual float GetParameter() const override;
+	INTERACTIVETOOLSFRAMEWORK_API virtual float GetParameter() const override;
 
-	virtual void SetParameter(float NewValue) override;
+	INTERACTIVETOOLSFRAMEWORK_API virtual void SetParameter(float NewValue) override;
 
-	virtual void BeginModify();
+	INTERACTIVETOOLSFRAMEWORK_API virtual void BeginModify();
 	
-	virtual void EndModify();
+	INTERACTIVETOOLSFRAMEWORK_API virtual void EndModify();
 
 
 public:
@@ -303,9 +303,9 @@ public:
 	/**
 	 * Create a standard instance of this ParameterSource, with the given AxisSource and TransformSource
 	 */
-	static UGizmoAxisIntervalParameterSource* Construct(
+	static INTERACTIVETOOLSFRAMEWORK_API UGizmoAxisIntervalParameterSource* Construct(
 		IGizmoFloatParameterSource* FloatSourceIn,
 		float ParameterMin = -FLT_MAX,
 		float ParameterMax = FLT_MAX,
 		UObject* Outer = (UObject*)GetTransientPackage());
-};
\ No newline at end of file
+};
diff --git a/Engine/Source/Runtime/InteractiveToolsFramework/Public/BaseGizmos/ParameterSourcesFloat.h b/Engine/Source/Runtime/InteractiveToolsFramework/Public/BaseGizmos/ParameterSourcesFloat.h
--- a/Engine/Source/Runtime/InteractiveToolsFramework/Public/BaseGizmos/ParameterSourcesFloat.h
+++ b/Engine/Source/Runtime/InteractiveToolsFramework/Public/BaseGizmos/ParameterSourcesFloat.h
@@ -17,7 +17,7 @@
  * IGizmoFloatParameterSource implementations use this to track changes and emit delta information.
  */
 USTRUCT()
-struct INTERACTIVETOOLSFRAMEWORK_API FGizmoFloatParameterChange
+struct FGizmoFloatParameterChange
 {
 	GENERATED_BODY()
 
@@ -44,8 +44,8 @@ struct INTERACTIVETOOLSFRAMEWORK_API FGizmoFloatParameterChange
  * UGizmoBaseFloatParameterSource is a base implementation of IGizmoFloatParameterSource,
  * which is not functional but adds an OnParameterChanged delegate for further subclasses.
  */
-UCLASS()
-class INTERACTIVETOOLSFRAMEWORK_API UGizmoBaseFloatParameterSource : public UObject, public IGizmoFloatParameterSource
+UCLASS(MinimalAPI)
+class UGizmoBaseFloatParameterSource : public UObject, public IGizmoFloatParameterSource
 {
 	GENERATED_BODY()
 public:
@@ -82,8 +82,8 @@ public:
  * (by way of UGizmoBaseFloatParameterSource) which locally stores the relevant Parameter
  * and emits update events via the OnParameterChanged delegate.
  */
-UCLASS()
-class INTERACTIVETOOLSFRAMEWORK_API UGizmoLocalFloatParameterSource : public UGizmoBaseFloatParameterSource
+UCLASS(MinimalAPI)
+class UGizmoLocalFloatParameterSource : public UGizmoBaseFloatParameterSource
 {
 	GENERATED_BODY()
 public:
diff --git a/Engine/Source/Runtime/InteractiveToolsFramework/Public/BaseGizmos/ParameterSourcesVec2.h b/Engine/Source/Runtime/InteractiveToolsFramework/Public/BaseGizmos/ParameterSourcesVec2.h
--- a/Engine/Source/Runtime/InteractiveToolsFramework/Public/BaseGizmos/ParameterSourcesVec2.h
+++ b/Engine/Source/Runtime/InteractiveToolsFramework/Public/BaseGizmos/ParameterSourcesVec2.h
@@ -16,7 +16,7 @@
  * IGizmoVec2ParameterSource implementations use this to track changes and emit delta information.
  */
 USTRUCT()
-struct INTERACTIVETOOLSFRAMEWORK_API FGizmoVec2ParameterChange
+struct FGizmoVec2ParameterChange
 {
 	GENERATED_BODY()
 
@@ -44,8 +44,8 @@ struct INTERACTIVETOOLSFRAMEWORK_API FGizmoVec2ParameterChange
  * UGizmoBaseVec2ParameterSource is a base implementation of IGizmoVec2ParameterSource,
  * which is not functional but adds an OnParameterChanged delegate for further subclasses.
  */
-UCLASS()
-class INTERACTIVETOOLSFRAMEWORK_API UGizmoBaseVec2ParameterSource : public UObject, public IGizmoVec2ParameterSource
+UCLASS(MinimalAPI)
+class UGizmoBaseVec2ParameterSource : public UObject, public IGizmoVec2ParameterSource
 {
 	GENERATED_BODY()
 public:
@@ -82,8 +82,8 @@ public:
  * (by way of UGizmoBaseVec2ParameterSource) which locally stores the relevant Parameter
  * and emits update events via the OnParameterChanged delegate.
  */
-UCLASS()
-class INTERACTIVETOOLSFRAMEWORK_API UGizmoLocalVec2ParameterSource : public UGizmoBaseVec2ParameterSource
+UCLASS(MinimalAPI)
+class UGizmoLocalVec2ParameterSource : public UGizmoBaseVec2ParameterSource
 {
 	GENERATED_BODY()
 public:
diff --git a/Engine/Source/Runtime/InteractiveToolsFramework/Public/BaseGizmos/ParameterToTransformAdapters.h b/Engine/Source/Runtime/InteractiveToolsFramework/Public/BaseGizmos/ParameterToTransformAdapters.h
--- a/Engine/Source/Runtime/InteractiveToolsFramework/Public/BaseGizmos/ParameterToTransformAdapters.h
+++ b/Engine/Source/Runtime/InteractiveToolsFramework/Public/BaseGizmos/ParameterToTransformAdapters.h
@@ -23,8 +23,8 @@
  *
  * This ParameterSource is intended to be used to create 3D Axis Translation Gizmos.
  */
-UCLASS()
-class INTERACTIVETOOLSFRAMEWORK_API UGizmoAxisTranslationParameterSource : public UGizmoBaseFloatParameterSource
+UCLASS(MinimalAPI)
+class UGizmoAxisTranslationParameterSource : public UGizmoBaseFloatParameterSource
 {
 	GENERATED_BODY()
 public:
@@ -50,10 +50,10 @@ public:
 		return Parameter;
 	}
 
-	virtual void SetParameter(float NewValue) override;
+	INTERACTIVETOOLSFRAMEWORK_API virtual void SetParameter(float NewValue) override;
 
-	virtual void BeginModify() override;
-	virtual void EndModify() override;
+	INTERACTIVETOOLSFRAMEWORK_API virtual void BeginModify() override;
+	INTERACTIVETOOLSFRAMEWORK_API virtual void EndModify() override;
 
 
 public:
@@ -115,8 +115,8 @@ public:
  * 
  * This ParameterSource is intended to be used to create 3D Plane Translation Gizmos.
  */
-UCLASS()
-class INTERACTIVETOOLSFRAMEWORK_API UGizmoPlaneTranslationParameterSource : public UGizmoBaseVec2ParameterSource
+UCLASS(MinimalAPI)
+class UGizmoPlaneTranslationParameterSource : public UGizmoBaseVec2ParameterSource
 {
 	GENERATED_BODY()
 public:
@@ -144,10 +144,10 @@ public:
 		return Parameter;
 	}
 
-	virtual void SetParameter(const FVector2D& NewValue) override;
+	INTERACTIVETOOLSFRAMEWORK_API virtual void SetParameter(const FVector2D& NewValue) override;
 
-	virtual void BeginModify() override;
-	virtual void EndModify() override;
+	INTERACTIVETOOLSFRAMEWORK_API virtual void BeginModify() override;
+	INTERACTIVETOOLSFRAMEWORK_API virtual void EndModify() override;
 
 
 public:
@@ -217,8 +217,8 @@ public:
  * around an IGizmoAxisSource (ie 3D axis). This rotation is applied to an IGizmoTransformSource.
  * This ParameterSource is intended to be used to create 3D Rotation Gizmos.
  */
-UCLASS()
-class INTERACTIVETOOLSFRAMEWORK_API UGizmoAxisRotationParameterSource : public UGizmoBaseFloatParameterSource
+UCLASS(MinimalAPI)
+class UGizmoAxisRotationParameterSource : public UGizmoBaseFloatParameterSource
 {
 	GENERATED_BODY()
 public:
@@ -242,10 +242,10 @@ public:
 		return Angle;
 	}
 
-	virtual void SetParameter(float NewValue) override;
+	INTERACTIVETOOLSFRAMEWORK_API virtual void SetParameter(float NewValue) override;
 
-	virtual void BeginModify() override;
-	virtual void EndModify() override;
+	INTERACTIVETOOLSFRAMEWORK_API virtual void BeginModify() override;
+	INTERACTIVETOOLSFRAMEWORK_API virtual void EndModify() override;
 
 
 public:
@@ -316,8 +316,8 @@ public:
   *
   * This ParameterSource is intended to be used to create 3D Uniform Scale Gizmos.
   */
-UCLASS()
-class INTERACTIVETOOLSFRAMEWORK_API UGizmoUniformScaleParameterSource : public UGizmoBaseVec2ParameterSource
+UCLASS(MinimalAPI)
+class UGizmoUniformScaleParameterSource : public UGizmoBaseVec2ParameterSource
 {
 	GENERATED_BODY()
 public:
@@ -336,10 +336,10 @@ public:
 		return Parameter;
 	}
 
-	virtual void SetParameter(const FVector2D& NewValue) override;
+	INTERACTIVETOOLSFRAMEWORK_API virtual void SetParameter(const FVector2D& NewValue) override;
 
-	virtual void BeginModify() override;
-	virtual void EndModify() override;
+	INTERACTIVETOOLSFRAMEWORK_API virtual void BeginModify() override;
+	INTERACTIVETOOLSFRAMEWORK_API virtual void EndModify() override;
 
 
 public:
@@ -412,8 +412,8 @@ public:
  *
  * This ParameterSource is intended to be used to create 3D Axis Scale Gizmos.
  */
-UCLASS()
-class INTERACTIVETOOLSFRAMEWORK_API UGizmoAxisScaleParameterSource : public UGizmoBaseFloatParameterSource
+UCLASS(MinimalAPI)
+class UGizmoAxisScaleParameterSource : public UGizmoBaseFloatParameterSource
 {
 	GENERATED_BODY()
 public:
@@ -431,10 +431,10 @@ public:
 		return Parameter;
 	}
 
-	virtual void SetParameter(float NewValue) override;
+	INTERACTIVETOOLSFRAMEWORK_API virtual void SetParameter(float NewValue) override;
 
-	virtual void BeginModify() override;
-	virtual void EndModify() override;
+	INTERACTIVETOOLSFRAMEWORK_API virtual void BeginModify() override;
+	INTERACTIVETOOLSFRAMEWORK_API virtual void EndModify() override;
 
 public:
 	/** The Parameter line-equation value is converted to a 3D Translation along this Axis */
@@ -505,8 +505,8 @@ public:
  *
  * This ParameterSource is intended to be used to create 3D Plane Scale Gizmos.
  */
-UCLASS()
-class INTERACTIVETOOLSFRAMEWORK_API UGizmoPlaneScaleParameterSource : public UGizmoBaseVec2ParameterSource
+UCLASS(MinimalAPI)
+class UGizmoPlaneScaleParameterSource : public UGizmoBaseVec2ParameterSource
 {
 	GENERATED_BODY()
 public:
@@ -525,10 +525,10 @@ public:
 		return Parameter;
 	}
 
-	virtual void SetParameter(const FVector2D& NewValue) override;
+	INTERACTIVETOOLSFRAMEWORK_API virtual void SetParameter(const FVector2D& NewValue) override;
 
-	virtual void BeginModify() override;
-	virtual void EndModify() override;
+	INTERACTIVETOOLSFRAMEWORK_API virtual void BeginModify() override;
+	INTERACTIVETOOLSFRAMEWORK_API virtual void EndModify() override;
 
 
 public:
diff --git a/Engine/Source/Runtime/InteractiveToolsFramework/Public/BaseGizmos/PlanePositionGizmo.h b/Engine/Source/Runtime/InteractiveToolsFramework/Public/BaseGizmos/PlanePositionGizmo.h
--- a/Engine/Source/Runtime/InteractiveToolsFramework/Public/BaseGizmos/PlanePositionGizmo.h
+++ b/Engine/Source/Runtime/InteractiveToolsFramework/Public/BaseGizmos/PlanePositionGizmo.h
@@ -29,13 +29,13 @@ class UObject;
 struct FToolBuilderState;
 
 
-UCLASS()
-class INTERACTIVETOOLSFRAMEWORK_API UPlanePositionGizmoBuilder : public UInteractiveGizmoBuilder
+UCLASS(MinimalAPI)
+class UPlanePositionGizmoBuilder : public UInteractiveGizmoBuilder
 {
 	GENERATED_BODY()
 
 public:
-	virtual UInteractiveGizmo* BuildGizmo(const FToolBuilderState& SceneState) const override;
+	INTERACTIVETOOLSFRAMEWORK_API virtual UInteractiveGizmo* BuildGizmo(const FToolBuilderState& SceneState) const override;
 };
 
 
@@ -56,29 +56,29 @@ public:
  *
  * Internally a UClickDragInputBehavior is used to handle mouse input, configured in ::Setup()
  */
-UCLASS()
-class INTERACTIVETOOLSFRAMEWORK_API UPlanePositionGizmo : public UInteractiveGizmo, public IClickDragBehaviorTarget, public IHoverBehaviorTarget
+UCLASS(MinimalAPI)
+class UPlanePositionGizmo : public UInteractiveGizmo, public IClickDragBehaviorTarget, public IHoverBehaviorTarget
 {
 	GENERATED_BODY()
 
 public:
 	// UInteractiveGizmo overrides
 
-	virtual void Setup() override;
+	INTERACTIVETOOLSFRAMEWORK_API virtual void Setup() override;
 
 	// IClickDragBehaviorTarget implementation
 
-	virtual FInputRayHit CanBeginClickDragSequence(const FInputDeviceRay& PressPos) override;
-	virtual void OnClickPress(const FInputDeviceRay& PressPos) override;
-	virtual void OnClickDrag(const FInputDeviceRay& DragPos) override;
-	virtual void OnClickRelease(const FInputDeviceRay& ReleasePos) override;
-	virtual void OnTerminateDragSequence() override;
+	INTERACTIVETOOLSFRAMEWORK_API virtual FInputRayHit CanBeginClickDragSequence(const FInputDeviceRay& PressPos) override;
+	INTERACTIVETOOLSFRAMEWORK_API virtual void OnClickPress(const FInputDeviceRay& PressPos) override;
+	INTERACTIVETOOLSFRAMEWORK_API virtual void OnClickDrag(const FInputDeviceRay& DragPos) override;
+	INTERACTIVETOOLSFRAMEWORK_API virtual void OnClickRelease(const FInputDeviceRay& ReleasePos) override;
+	INTERACTIVETOOLSFRAMEWORK_API virtual void OnTerminateDragSequence() override;
 
 	// IHoverBehaviorTarget implementation
-	virtual FInputRayHit BeginHoverSequenceHitTest(const FInputDeviceRay& PressPos) override;
-	virtual void OnBeginHover(const FInputDeviceRay& DevicePos) override;
-	virtual bool OnUpdateHover(const FInputDeviceRay& DevicePos) override;
-	virtual void OnEndHover() override;
+	INTERACTIVETOOLSFRAMEWORK_API virtual FInputRayHit BeginHoverSequenceHitTest(const FInputDeviceRay& PressPos) override;
+	INTERACTIVETOOLSFRAMEWORK_API virtual void OnBeginHover(const FInputDeviceRay& DevicePos) override;
+	INTERACTIVETOOLSFRAMEWORK_API virtual bool OnUpdateHover(const FInputDeviceRay& DevicePos) override;
+	INTERACTIVETOOLSFRAMEWORK_API virtual void OnEndHover() override;
 
 public:
 	/** AxisSource provides the 3D plane on which the interaction happens */
diff --git a/Engine/Source/Runtime/InteractiveToolsFramework/Public/BaseGizmos/RepositionableTransformGizmo.h b/Engine/Source/Runtime/InteractiveToolsFramework/Public/BaseGizmos/RepositionableTransformGizmo.h
--- a/Engine/Source/Runtime/InteractiveToolsFramework/Public/BaseGizmos/RepositionableTransformGizmo.h
+++ b/Engine/Source/Runtime/InteractiveToolsFramework/Public/BaseGizmos/RepositionableTransformGizmo.h
@@ -13,22 +13,22 @@
 class IToolContextTransactionProvider;
 class UTransformProxy;
 
-UCLASS()
-class INTERACTIVETOOLSFRAMEWORK_API URepositionableTransformGizmoBuilder : public UCombinedTransformGizmoBuilder
+UCLASS(MinimalAPI)
+class URepositionableTransformGizmoBuilder : public UCombinedTransformGizmoBuilder
 {
 	GENERATED_BODY()
 
 public:
 
-	virtual UInteractiveGizmo* BuildGizmo(const FToolBuilderState& SceneState) const override;
+	INTERACTIVETOOLSFRAMEWORK_API virtual UInteractiveGizmo* BuildGizmo(const FToolBuilderState& SceneState) const override;
 };
 
 
 /**
  * A transform gizmo that also allows the user to reposition it by middle-clicking rotation/translation components.
  */
-UCLASS()
-class INTERACTIVETOOLSFRAMEWORK_API URepositionableTransformGizmo : public UCombinedTransformGizmo
+UCLASS(MinimalAPI)
+class URepositionableTransformGizmo : public UCombinedTransformGizmo
 {
 	GENERATED_BODY()
 
@@ -39,13 +39,13 @@ public:
 	 * @param Target active target
 	 * @param TransactionProvider optional IToolContextTransactionProvider implementation to use - by default uses GizmoManager
 	 */
-	virtual void SetActiveTarget(UTransformProxy* Target, IToolContextTransactionProvider* TransactionProvider = nullptr) override;
+	INTERACTIVETOOLSFRAMEWORK_API virtual void SetActiveTarget(UTransformProxy* Target, IToolContextTransactionProvider* TransactionProvider = nullptr) override;
 
 	/**
 	 * Allows the user to provide functions to use for aligning the gizmo destination to items in the scene. Note that this affects
 	 * both the movement and repositioning.
 	 */
-	virtual void SetWorldAlignmentFunctions(
+	INTERACTIVETOOLSFRAMEWORK_API virtual void SetWorldAlignmentFunctions(
 		TUniqueFunction<bool()>&& ShouldAlignTranslationIn,
 		TUniqueFunction<bool(const FRay&, FVector&)>&& TranslationAlignmentRayCasterIn) override;
 
@@ -54,7 +54,7 @@ public:
 	 * This is because the repositining is likely to want different raycasts that don't ignore the 
 	 * moved components themselves, so that the gizmo can be aligned to a target component.
 	 */
-	virtual void SetPivotAlignmentFunctions(
+	INTERACTIVETOOLSFRAMEWORK_API virtual void SetPivotAlignmentFunctions(
 		TUniqueFunction<bool()>&& ShouldAlignPivotIn,
 		TUniqueFunction<bool(const FRay&, FVector&)>&& PivotAlignmentRayCasterIn);
 protected:
@@ -69,7 +69,7 @@ protected:
 	TObjectPtr<UGizmoTransformChangeStateTarget> RepositionStateTarget;
 
 	// Helper functions
-	void ModifyPivotAxisGizmo(UInteractiveGizmo* SubGizmoIn);
-	void ModifyPivotPlaneGizmo(UInteractiveGizmo* SubGizmoIn);
-	void ModifyPivotRotateGizmo(UInteractiveGizmo* SubGizmoIn);
+	INTERACTIVETOOLSFRAMEWORK_API void ModifyPivotAxisGizmo(UInteractiveGizmo* SubGizmoIn);
+	INTERACTIVETOOLSFRAMEWORK_API void ModifyPivotPlaneGizmo(UInteractiveGizmo* SubGizmoIn);
+	INTERACTIVETOOLSFRAMEWORK_API void ModifyPivotRotateGizmo(UInteractiveGizmo* SubGizmoIn);
 };
diff --git a/Engine/Source/Runtime/InteractiveToolsFramework/Public/BaseGizmos/ScalableSphereGizmo.h b/Engine/Source/Runtime/InteractiveToolsFramework/Public/BaseGizmos/ScalableSphereGizmo.h
--- a/Engine/Source/Runtime/InteractiveToolsFramework/Public/BaseGizmos/ScalableSphereGizmo.h
+++ b/Engine/Source/Runtime/InteractiveToolsFramework/Public/BaseGizmos/ScalableSphereGizmo.h
@@ -11,49 +11,49 @@
 
 #include "ScalableSphereGizmo.generated.h"
 
-UCLASS()
-class INTERACTIVETOOLSFRAMEWORK_API UScalableSphereGizmoBuilder : public UInteractiveGizmoBuilder
+UCLASS(MinimalAPI)
+class UScalableSphereGizmoBuilder : public UInteractiveGizmoBuilder
 {
 	GENERATED_BODY()
 
 public:
-	virtual UInteractiveGizmo* BuildGizmo(const FToolBuilderState& SceneState) const override;
+	INTERACTIVETOOLSFRAMEWORK_API virtual UInteractiveGizmo* BuildGizmo(const FToolBuilderState& SceneState) const override;
 };
 
 /**
  * UScalableSphereGizmo provides a sphere that can be scaled in all directions by dragging
  * anywhere on the three axial circles that represent it
  */
-UCLASS()
-class INTERACTIVETOOLSFRAMEWORK_API UScalableSphereGizmo : public UInteractiveGizmo, public IHoverBehaviorTarget
+UCLASS(MinimalAPI)
+class UScalableSphereGizmo : public UInteractiveGizmo, public IHoverBehaviorTarget
 {
 	GENERATED_BODY()
 
 public:
 	// UInteractiveGizmo interface
-	virtual void Setup() override;
-	virtual void Render(IToolsContextRenderAPI* RenderAPI) override;
+	INTERACTIVETOOLSFRAMEWORK_API virtual void Setup() override;
+	INTERACTIVETOOLSFRAMEWORK_API virtual void Render(IToolsContextRenderAPI* RenderAPI) override;
 
 	// IHoverBehaviorTarget interface
-	virtual FInputRayHit BeginHoverSequenceHitTest(const FInputDeviceRay& PressPos) override;
+	INTERACTIVETOOLSFRAMEWORK_API virtual FInputRayHit BeginHoverSequenceHitTest(const FInputDeviceRay& PressPos) override;
 	virtual void OnBeginHover(const FInputDeviceRay& DevicePos) override { }
-	virtual bool OnUpdateHover(const FInputDeviceRay& DevicePos);
-	virtual void OnEndHover();
+	INTERACTIVETOOLSFRAMEWORK_API virtual bool OnUpdateHover(const FInputDeviceRay& DevicePos);
+	INTERACTIVETOOLSFRAMEWORK_API virtual void OnEndHover();
 	
 	/**
 	 * Set the Target to which the gizmo will be attached
 	 */
-	virtual void SetTarget(UTransformProxy* InTarget);
+	INTERACTIVETOOLSFRAMEWORK_API virtual void SetTarget(UTransformProxy* InTarget);
 
-	virtual void OnBeginDrag(const FInputDeviceRay& Ray);
-	virtual void OnUpdateDrag(const FInputDeviceRay& Ray);
-	virtual void OnEndDrag(const FInputDeviceRay& Ray);
+	INTERACTIVETOOLSFRAMEWORK_API virtual void OnBeginDrag(const FInputDeviceRay& Ray);
+	INTERACTIVETOOLSFRAMEWORK_API virtual void OnUpdateDrag(const FInputDeviceRay& Ray);
+	INTERACTIVETOOLSFRAMEWORK_API virtual void OnEndDrag(const FInputDeviceRay& Ray);
 
 	/** Check if the input Ray hit any of the components of the internal actor */
-	bool HitTest(const FRay& Ray, FHitResult& OutHit, FVector& OutAxis, FTransform& OutTransform);
+	INTERACTIVETOOLSFRAMEWORK_API bool HitTest(const FRay& Ray, FHitResult& OutHit, FVector& OutAxis, FTransform& OutTransform);
 
 	/* Set the Radius of the Sphere*/
-	void SetRadius(float InRadius);
+	INTERACTIVETOOLSFRAMEWORK_API void SetRadius(float InRadius);
 
 	/** Called when the radius is chaged (by dragging or setting). Sends new radius as parameter. */
 	TFunction<void(const float)> UpdateRadiusFunc = nullptr;
@@ -69,7 +69,7 @@ public:
 private:
 
 	// Check if the Circle represented by CircleNormal centered at the gizmos location has been hit by the Ray
-	bool CheckCircleIntersection(const FRay& Ray, FVector CircleNormal, FVector& OutHitLocation, FVector& OutHitAxis);
+	INTERACTIVETOOLSFRAMEWORK_API bool CheckCircleIntersection(const FRay& Ray, FVector CircleNormal, FVector& OutHitLocation, FVector& OutHitAxis);
 	
 	// The radius of the sphere
 	UPROPERTY()
@@ -107,24 +107,24 @@ private:
 /**
  * A behavior that forwards clicking and dragging to the gizmo.
  */
-UCLASS()
-class INTERACTIVETOOLSFRAMEWORK_API UScalableSphereGizmoInputBehavior : public UAnyButtonInputBehavior
+UCLASS(MinimalAPI)
+class UScalableSphereGizmoInputBehavior : public UAnyButtonInputBehavior
 {
 	GENERATED_BODY()
 
 public:
 	virtual FInputCapturePriority GetPriority() override { return FInputCapturePriority(FInputCapturePriority::DEFAULT_GIZMO_PRIORITY); }
 
-	virtual void Initialize(UScalableSphereGizmo* Gizmo);
+	INTERACTIVETOOLSFRAMEWORK_API virtual void Initialize(UScalableSphereGizmo* Gizmo);
 
-	virtual FInputCaptureRequest WantsCapture(const FInputDeviceState& input) override;
-	virtual FInputCaptureUpdate BeginCapture(const FInputDeviceState& input, EInputCaptureSide eSide) override;
-	virtual FInputCaptureUpdate UpdateCapture(const FInputDeviceState& input, const FInputCaptureData& data) override;
-	virtual void ForceEndCapture(const FInputCaptureData& data) override;
+	INTERACTIVETOOLSFRAMEWORK_API virtual FInputCaptureRequest WantsCapture(const FInputDeviceState& input) override;
+	INTERACTIVETOOLSFRAMEWORK_API virtual FInputCaptureUpdate BeginCapture(const FInputDeviceState& input, EInputCaptureSide eSide) override;
+	INTERACTIVETOOLSFRAMEWORK_API virtual FInputCaptureUpdate UpdateCapture(const FInputDeviceState& input, const FInputCaptureData& data) override;
+	INTERACTIVETOOLSFRAMEWORK_API virtual void ForceEndCapture(const FInputCaptureData& data) override;
 
 protected:
 	UScalableSphereGizmo* Gizmo;
 	FRay LastWorldRay;
 	FVector2D LastScreenPosition;
 	bool bInputDragCaptured;
-};
\ No newline at end of file
+};
diff --git a/Engine/Source/Runtime/InteractiveToolsFramework/Public/BaseGizmos/StateTargets.h b/Engine/Source/Runtime/InteractiveToolsFramework/Public/BaseGizmos/StateTargets.h
--- a/Engine/Source/Runtime/InteractiveToolsFramework/Public/BaseGizmos/StateTargets.h
+++ b/Engine/Source/Runtime/InteractiveToolsFramework/Public/BaseGizmos/StateTargets.h
@@ -30,8 +30,8 @@
 /**
  * UGizmoNilStateTarget is an implementation of IGizmoStateTarget that does nothing
  */
-UCLASS()
-class INTERACTIVETOOLSFRAMEWORK_API UGizmoNilStateTarget : public UObject, public IGizmoStateTarget
+UCLASS(MinimalAPI)
+class UGizmoNilStateTarget : public UObject, public IGizmoStateTarget
 {
 	GENERATED_BODY()
 public:
@@ -51,8 +51,8 @@ public:
  * UGizmoLambdaStateTarget is an implementation of IGizmoStateTarget that forwards
  * calls to its interface functions to external TFunctions
  */
-UCLASS()
-class INTERACTIVETOOLSFRAMEWORK_API UGizmoLambdaStateTarget : public UObject, public IGizmoStateTarget
+UCLASS(MinimalAPI)
+class UGizmoLambdaStateTarget : public UObject, public IGizmoStateTarget
 {
 	GENERATED_BODY()
 public:
@@ -82,8 +82,8 @@ public:
  * UGizmoObjectModifyStateTarget is an implementation of IGizmoStateTarget that 
  * opens and closes change transactions on a target UObject via a GizmoManager.
  */
-UCLASS()
-class INTERACTIVETOOLSFRAMEWORK_API UGizmoObjectModifyStateTarget : public UObject, public IGizmoStateTarget
+UCLASS(MinimalAPI)
+class UGizmoObjectModifyStateTarget : public UObject, public IGizmoStateTarget
 {
 	GENERATED_BODY()
 public:
@@ -163,8 +163,8 @@ public:
  * The DependentChangeSources and ExternalDependentChangeSources lists allow additional
  * FChange objects to be inserted into the transaction, provided by IToolCommandChangeSource implementations.
  */
-UCLASS()
-class INTERACTIVETOOLSFRAMEWORK_API UGizmoTransformChangeStateTarget : public UObject, public IGizmoStateTarget
+UCLASS(MinimalAPI)
+class UGizmoTransformChangeStateTarget : public UObject, public IGizmoStateTarget
 {
 	GENERATED_BODY()
 public:
diff --git a/Engine/Source/Runtime/InteractiveToolsFramework/Public/BaseGizmos/TransformGizmoUtil.h b/Engine/Source/Runtime/InteractiveToolsFramework/Public/BaseGizmos/TransformGizmoUtil.h
--- a/Engine/Source/Runtime/InteractiveToolsFramework/Public/BaseGizmos/TransformGizmoUtil.h
+++ b/Engine/Source/Runtime/InteractiveToolsFramework/Public/BaseGizmos/TransformGizmoUtil.h
@@ -84,23 +84,23 @@ namespace TransformGizmoUtil
  * will look up the UCombinedTransformGizmoContextObject instance in the ContextObjectStore and then
  * call the associated function below.
  */
-UCLASS(Transient)
-class INTERACTIVETOOLSFRAMEWORK_API UCombinedTransformGizmoContextObject : public UObject
+UCLASS(Transient, MinimalAPI)
+class UCombinedTransformGizmoContextObject : public UObject
 {
 	GENERATED_BODY()
 public:
 
 public:
 	// builder identifiers for default gizmo types. Perhaps should have an API for this...
-	static const FString DefaultAxisPositionBuilderIdentifier;
-	static const FString DefaultPlanePositionBuilderIdentifier;
-	static const FString DefaultAxisAngleBuilderIdentifier;
-	static const FString DefaultThreeAxisTransformBuilderIdentifier;
-	static const FString CustomThreeAxisTransformBuilderIdentifier;
-	static const FString CustomRepositionableThreeAxisTransformBuilderIdentifier;
+	static INTERACTIVETOOLSFRAMEWORK_API const FString DefaultAxisPositionBuilderIdentifier;
+	static INTERACTIVETOOLSFRAMEWORK_API const FString DefaultPlanePositionBuilderIdentifier;
+	static INTERACTIVETOOLSFRAMEWORK_API const FString DefaultAxisAngleBuilderIdentifier;
+	static INTERACTIVETOOLSFRAMEWORK_API const FString DefaultThreeAxisTransformBuilderIdentifier;
+	static INTERACTIVETOOLSFRAMEWORK_API const FString CustomThreeAxisTransformBuilderIdentifier;
+	static INTERACTIVETOOLSFRAMEWORK_API const FString CustomRepositionableThreeAxisTransformBuilderIdentifier;
 
-	void RegisterGizmosWithManager(UInteractiveToolManager* ToolManager);
-	void DeregisterGizmosWithManager(UInteractiveToolManager* ToolManager);
+	INTERACTIVETOOLSFRAMEWORK_API void RegisterGizmosWithManager(UInteractiveToolManager* ToolManager);
+	INTERACTIVETOOLSFRAMEWORK_API void DeregisterGizmosWithManager(UInteractiveToolManager* ToolManager);
 
 	/**
 	 * Activate a new instance of the default 3-axis transformation Gizmo. RegisterDefaultGizmos() must have been called first.
@@ -108,7 +108,7 @@ public:
 	 * @param InstanceIdentifier optional client-defined *unique* string that can be used to locate this instance
 	 * @return new Gizmo instance that has been created and initialized
 	 */
-	virtual UCombinedTransformGizmo* Create3AxisTransformGizmo(UInteractiveGizmoManager* GizmoManager, void* Owner = nullptr, const FString& InstanceIdentifier = FString());
+	INTERACTIVETOOLSFRAMEWORK_API virtual UCombinedTransformGizmo* Create3AxisTransformGizmo(UInteractiveGizmoManager* GizmoManager, void* Owner = nullptr, const FString& InstanceIdentifier = FString());
 
 	/**
 	 * Activate a new customized instance of the default 3-axis transformation Gizmo, with only certain elements included. RegisterDefaultGizmos() must have been called first.
@@ -117,13 +117,13 @@ public:
 	 * @param InstanceIdentifier optional client-defined *unique* string that can be used to locate this instance
 	 * @return new Gizmo instance that has been created and initialized
 	 */
-	virtual UCombinedTransformGizmo* CreateCustomTransformGizmo(UInteractiveGizmoManager* GizmoManager, ETransformGizmoSubElements Elements, void* Owner = nullptr, const FString& InstanceIdentifier = FString());
+	INTERACTIVETOOLSFRAMEWORK_API virtual UCombinedTransformGizmo* CreateCustomTransformGizmo(UInteractiveGizmoManager* GizmoManager, ETransformGizmoSubElements Elements, void* Owner = nullptr, const FString& InstanceIdentifier = FString());
 
 	/**
 	 * Variant of CreateCustomTransformGizmo that creates a URepositionableTransformGizmo, which is an extension to UCombinedTransformGizmo that 
 	 * supports various snapping interactions
 	 */
-	virtual UCombinedTransformGizmo* CreateCustomRepositionableTransformGizmo(UInteractiveGizmoManager* GizmoManager, ETransformGizmoSubElements Elements, void* Owner = nullptr, const FString& InstanceIdentifier = FString());
+	INTERACTIVETOOLSFRAMEWORK_API virtual UCombinedTransformGizmo* CreateCustomRepositionableTransformGizmo(UInteractiveGizmoManager* GizmoManager, ETransformGizmoSubElements Elements, void* Owner = nullptr, const FString& InstanceIdentifier = FString());
 
 	DECLARE_MULTICAST_DELEGATE_OneParam(FOnGizmoCreated, UCombinedTransformGizmo*);
 	FOnGizmoCreated OnGizmoCreated;
diff --git a/Engine/Source/Runtime/InteractiveToolsFramework/Public/BaseGizmos/TransformProxy.h b/Engine/Source/Runtime/InteractiveToolsFramework/Public/BaseGizmos/TransformProxy.h
--- a/Engine/Source/Runtime/InteractiveToolsFramework/Public/BaseGizmos/TransformProxy.h
+++ b/Engine/Source/Runtime/InteractiveToolsFramework/Public/BaseGizmos/TransformProxy.h
@@ -35,8 +35,8 @@ class USceneComponent;
  * Otherwise the main transform is centered at the average origin and
  * has no rotation.
  */
-UCLASS(Transient)
-class INTERACTIVETOOLSFRAMEWORK_API UTransformProxy : public UObject
+UCLASS(Transient, MinimalAPI)
+class UTransformProxy : public UObject
 {
 	GENERATED_BODY()
 public:
@@ -46,7 +46,7 @@ public:
 	 * @param bModifyComponentOnTransform if true, Component->Modify() is called before the Component transform is updated
 	 * @warning The internal shared transform is regenerated each time a component is added
 	 */
-	virtual void AddComponent(USceneComponent* Component, bool bModifyComponentOnTransform = true);
+	INTERACTIVETOOLSFRAMEWORK_API virtual void AddComponent(USceneComponent* Component, bool bModifyComponentOnTransform = true);
 
 	/**
 	 * Add a component sub-object to the proxy set with custom transform access functions. This can be used
@@ -57,7 +57,7 @@ public:
 	 * @param bModifyComponentOnTransform if true, Component->Modify() is called before the Component transform is updated
 	 * @warning The internal shared transform is regenerated each time a component is added
 	 */
-	virtual void AddComponentCustom(
+	INTERACTIVETOOLSFRAMEWORK_API virtual void AddComponentCustom(
 		USceneComponent* Component, 
 		TUniqueFunction<FTransform(void)> GetTransformFunc,
 		TUniqueFunction<void(const FTransform&)> SetTransformFunc,
@@ -68,12 +68,12 @@ public:
 	/**
 	 * @return the shared transform for all the sub-objects
 	 */
-	virtual FTransform GetTransform() const;
+	INTERACTIVETOOLSFRAMEWORK_API virtual FTransform GetTransform() const;
 
 	/**
 	 * Update the main transform and then update the sub-objects based on their relative transformations
 	 */
-	virtual void SetTransform(const FTransform& Transform);
+	INTERACTIVETOOLSFRAMEWORK_API virtual void SetTransform(const FTransform& Transform);
 
 
 	/**
@@ -83,24 +83,24 @@ public:
 	 * as the TransformProxy doesn't know about this external state. 
 	 * @note FTransformProxyChangeSource will call these functions on Begin/End
 	 */
-	virtual void BeginTransformEditSequence();
+	INTERACTIVETOOLSFRAMEWORK_API virtual void BeginTransformEditSequence();
 
 	/**
 	 * External clients should call this when done a sequence of SetTransform calls (see BeginTransformEditSequence)
 	 */
-	virtual void EndTransformEditSequence();
+	INTERACTIVETOOLSFRAMEWORK_API virtual void EndTransformEditSequence();
 
 	/**
 	 * Clients should call this before a sequence of SetTransform calls that have bSetPivotMode as true
 	 * (see comment in BeginTransformEditSequence).
 	 */
-	virtual void BeginPivotEditSequence();
+	INTERACTIVETOOLSFRAMEWORK_API virtual void BeginPivotEditSequence();
 
 	/**
 	 * Clients should call this when done with a sequence of SetTransform that have bSetPivotMode as true 
 	 * (see comment in BeginTransformEditSequence).
 	 */
-	virtual void EndPivotEditSequence();
+	INTERACTIVETOOLSFRAMEWORK_API virtual void EndPivotEditSequence();
 
 public:
 	/**
@@ -182,13 +182,13 @@ protected:
 	FTransform InitialSharedTransform;
 
 	/** Recalculate main SharedTransform when object set changes*/
-	virtual void UpdateSharedTransform();
+	INTERACTIVETOOLSFRAMEWORK_API virtual void UpdateSharedTransform();
 
 	/** Recalculate per-object relative transforms */
-	virtual void UpdateObjectTransforms();
+	INTERACTIVETOOLSFRAMEWORK_API virtual void UpdateObjectTransforms();
 
 	/** Propagate a transform update to the sub-objects */
-	virtual void UpdateObjects();
+	INTERACTIVETOOLSFRAMEWORK_API virtual void UpdateObjects();
 };
 
 
@@ -196,15 +196,15 @@ protected:
 /**
  * FTransformProxyChange tracks a change to the base transform for a TransformProxy
  */
-class INTERACTIVETOOLSFRAMEWORK_API FTransformProxyChange : public FToolCommandChange
+class FTransformProxyChange : public FToolCommandChange
 {
 public:
 	FTransform From;
 	FTransform To;
 	bool bSetPivotMode = false;
 
-	virtual void Apply(UObject* Object) override;
-	virtual void Revert(UObject* Object) override;
+	INTERACTIVETOOLSFRAMEWORK_API virtual void Apply(UObject* Object) override;
+	INTERACTIVETOOLSFRAMEWORK_API virtual void Revert(UObject* Object) override;
 
 	virtual FString ToString() const override { return TEXT("FTransformProxyChange"); }
 };
@@ -213,7 +213,7 @@ public:
  * FTransformProxyChangeSource generates FTransformProxyChange instances on Begin/End.
  * Instances of this class can (for example) be attached to a UGizmoTransformChangeStateTarget for use TransformGizmo change tracking.
  */
-class INTERACTIVETOOLSFRAMEWORK_API FTransformProxyChangeSource : public IToolCommandChangeSource
+class FTransformProxyChangeSource : public IToolCommandChangeSource
 {
 public:
 	FTransformProxyChangeSource(UTransformProxy* ProxyIn)
@@ -226,10 +226,10 @@ public:
 	TWeakObjectPtr<UTransformProxy> Proxy;
 	TUniquePtr<FTransformProxyChange> ActiveChange;
 
-	virtual void BeginChange() override;
-	virtual TUniquePtr<FToolCommandChange> EndChange() override;
-	virtual UObject* GetChangeTarget() override;
-	virtual FText GetChangeDescription() override;
+	INTERACTIVETOOLSFRAMEWORK_API virtual void BeginChange() override;
+	INTERACTIVETOOLSFRAMEWORK_API virtual TUniquePtr<FToolCommandChange> EndChange() override;
+	INTERACTIVETOOLSFRAMEWORK_API virtual UObject* GetChangeTarget() override;
+	INTERACTIVETOOLSFRAMEWORK_API virtual FText GetChangeDescription() override;
 
 	/**
 	 * If true, the emitted changes will always have bSetPivotMode set to true, regardless of
@@ -238,4 +238,4 @@ public:
 	 * otherwise behaves normally.
 	 */
 	bool bOverrideSetPivotMode = false;
-};
\ No newline at end of file
+};
diff --git a/Engine/Source/Runtime/InteractiveToolsFramework/Public/BaseGizmos/TransformSources.h b/Engine/Source/Runtime/InteractiveToolsFramework/Public/BaseGizmos/TransformSources.h
--- a/Engine/Source/Runtime/InteractiveToolsFramework/Public/BaseGizmos/TransformSources.h
+++ b/Engine/Source/Runtime/InteractiveToolsFramework/Public/BaseGizmos/TransformSources.h
@@ -30,8 +30,8 @@ class UTransformProxy;
  * UGizmoBaseTransformSource is a base implementation of IGizmoTransformSource that 
  * adds an OnTransformChanged delegate. This class cannot be used directly and must be subclassed.
  */
-UCLASS()
-class INTERACTIVETOOLSFRAMEWORK_API UGizmoBaseTransformSource : public UObject, public IGizmoTransformSource
+UCLASS(MinimalAPI)
+class UGizmoBaseTransformSource : public UObject, public IGizmoTransformSource
 {
 	GENERATED_BODY()
 public:
@@ -57,15 +57,15 @@ public:
  * UGizmoComponentWorldTransformSource implements IGizmoTransformSource (via UGizmoBaseTransformSource)
  * based on the internal transform of a USceneComponent.
  */
-UCLASS()
-class INTERACTIVETOOLSFRAMEWORK_API UGizmoComponentWorldTransformSource : public UGizmoBaseTransformSource
+UCLASS(MinimalAPI)
+class UGizmoComponentWorldTransformSource : public UGizmoBaseTransformSource
 {
 	GENERATED_BODY()
 public:
 
-	virtual FTransform GetTransform() const override;
+	INTERACTIVETOOLSFRAMEWORK_API virtual FTransform GetTransform() const override;
 
-	virtual void SetTransform(const FTransform& NewTransform) override;
+	INTERACTIVETOOLSFRAMEWORK_API virtual void SetTransform(const FTransform& NewTransform) override;
 
 	UPROPERTY()
 	TObjectPtr<USceneComponent> Component;
@@ -116,8 +116,8 @@ struct FSeparateScaleProvider
  * an external source via FSeparateScaleProvider, and in GetTransform/SetTransform rewrites the
  * Transform from the child IGizmoTransformSource with the new scale.
  */
-UCLASS()
-class INTERACTIVETOOLSFRAMEWORK_API UGizmoScaledTransformSource : public UGizmoBaseTransformSource
+UCLASS(MinimalAPI)
+class UGizmoScaledTransformSource : public UGizmoBaseTransformSource
 {
 	GENERATED_BODY()
 public:
@@ -125,12 +125,12 @@ public:
 	/**
 	 * IGizmoTransformSource implementation, returns child transform with local sclae
 	 */
-	virtual FTransform GetTransform() const override;
+	INTERACTIVETOOLSFRAMEWORK_API virtual FTransform GetTransform() const override;
 
 	/**
 	 * IGizmoTransformSource implementation, removes scale and sends to ScaleProvider, then forwards remaining rotate+translate transform to child
 	 */
-	virtual void SetTransform(const FTransform& NewTransform) override;
+	INTERACTIVETOOLSFRAMEWORK_API virtual void SetTransform(const FTransform& NewTransform) override;
 
 	/**
 	 * Child transform source
@@ -146,7 +146,7 @@ public:
 	/**
 	 * Return the child transform with combined scale
 	 */
-	FTransform GetScaledTransform() const;
+	INTERACTIVETOOLSFRAMEWORK_API FTransform GetScaledTransform() const;
 
 public:
 	/**
@@ -173,15 +173,15 @@ public:
  * UGizmoTransformProxyTransformSource implements IGizmoTransformSource (via UGizmoBaseTransformSource)
  * based on the internal transform of a UTransformProxy.
  */
-UCLASS()
-class INTERACTIVETOOLSFRAMEWORK_API UGizmoTransformProxyTransformSource : public UGizmoBaseTransformSource
+UCLASS(MinimalAPI)
+class UGizmoTransformProxyTransformSource : public UGizmoBaseTransformSource
 {
 	GENERATED_BODY()
 public:
 
-	virtual FTransform GetTransform() const override;
+	INTERACTIVETOOLSFRAMEWORK_API virtual FTransform GetTransform() const override;
 
-	virtual void SetTransform(const FTransform& NewTransform) override;
+	INTERACTIVETOOLSFRAMEWORK_API virtual void SetTransform(const FTransform& NewTransform) override;
 
 	UPROPERTY()
 	TObjectPtr<UTransformProxy> Proxy;
@@ -218,17 +218,17 @@ public:
  * the unscaled transform to a gizmo component (since we don't want to scale the gizmo component but do want to 
  * rotate/translate it).
  */
-UCLASS()
-class INTERACTIVETOOLSFRAMEWORK_API UGizmoScaledAndUnscaledTransformSources : public UGizmoBaseTransformSource
+UCLASS(MinimalAPI)
+class UGizmoScaledAndUnscaledTransformSources : public UGizmoBaseTransformSource
 {
 	GENERATED_BODY()
 public:
 
 	/** Gets the transform from ScaledTransformSource. */
-	virtual FTransform GetTransform() const override;
+	INTERACTIVETOOLSFRAMEWORK_API virtual FTransform GetTransform() const override;
 
 	/** Calls SetTransform on ScaledTransformSource and passes the unscaled version to UnscaledTransformSource. */
-	virtual void SetTransform(const FTransform& NewTransform) override;
+	INTERACTIVETOOLSFRAMEWORK_API virtual void SetTransform(const FTransform& NewTransform) override;
 
 	UPROPERTY()
 	TScriptInterface<IGizmoTransformSource> ScaledTransformSource;
