[Engine]

diff --git a/Engine/Source/Runtime/InteractiveToolsFramework/Public/BaseTools/BaseBrushTool.h b/Engine/Source/Runtime/InteractiveToolsFramework/Public/BaseTools/BaseBrushTool.h
--- a/Engine/Source/Runtime/InteractiveToolsFramework/Public/BaseTools/BaseBrushTool.h
+++ b/Engine/Source/Runtime/InteractiveToolsFramework/Public/BaseTools/BaseBrushTool.h
@@ -15,13 +15,13 @@ class UBrushStampIndicator;
 /**
  * Standard properties for a Brush-type Tool
  */
-UCLASS()
-class INTERACTIVETOOLSFRAMEWORK_API UBrushBaseProperties : public UInteractiveToolPropertySet
+UCLASS(MinimalAPI)
+class UBrushBaseProperties : public UInteractiveToolPropertySet
 {
 	GENERATED_BODY()
 
 public:
-	UBrushBaseProperties();
+	INTERACTIVETOOLSFRAMEWORK_API UBrushBaseProperties();
 
 	/** Relative size of brush */
 	UPROPERTY(EditAnywhere, Category = Brush, meta = (DisplayName = "Size", UIMin = "0.0", UIMax = "1.0", ClampMin = "0.0", ClampMax = "10.0", DisplayPriority = 1, HideEditConditionToggle, EditCondition = "bSpecifyRadius == false"))
@@ -58,7 +58,7 @@ public:
  * Generic Brush Stamp data
  */
 USTRUCT()
-struct INTERACTIVETOOLSFRAMEWORK_API FBrushStampData
+struct FBrushStampData
 {
 	GENERATED_BODY();
 	/** Radius of brush stamp */
@@ -148,34 +148,34 @@ private:
  *   4) status of brush stroke via .bInBrushStroke FProperty
  *   5) "B" hotkey to adjust brush radius / strength by click-dragging in the viewport
  */
-UCLASS()
-class INTERACTIVETOOLSFRAMEWORK_API UBaseBrushTool : public UMeshSurfacePointTool
+UCLASS(MinimalAPI)
+class UBaseBrushTool : public UMeshSurfacePointTool
 {
 	GENERATED_BODY()
 
 public:
-	UBaseBrushTool();
+	INTERACTIVETOOLSFRAMEWORK_API UBaseBrushTool();
 
-	virtual void RegisterActions(FInteractiveToolActionSet& ActionSet) override;
+	INTERACTIVETOOLSFRAMEWORK_API virtual void RegisterActions(FInteractiveToolActionSet& ActionSet) override;
 
 	// UInteractiveTool
-	virtual void Setup() override;
-	virtual void Shutdown(EToolShutdownType ShutdownType) override;
-	virtual void Render(IToolsContextRenderAPI* RenderAPI) override;
-	virtual void DrawHUD( FCanvas* Canvas, IToolsContextRenderAPI* RenderAPI ) override;
+	INTERACTIVETOOLSFRAMEWORK_API virtual void Setup() override;
+	INTERACTIVETOOLSFRAMEWORK_API virtual void Shutdown(EToolShutdownType ShutdownType) override;
+	INTERACTIVETOOLSFRAMEWORK_API virtual void Render(IToolsContextRenderAPI* RenderAPI) override;
+	INTERACTIVETOOLSFRAMEWORK_API virtual void DrawHUD( FCanvas* Canvas, IToolsContextRenderAPI* RenderAPI ) override;
 	
-	virtual void OnPropertyModified(UObject* PropertySet, FProperty* Property) override;
+	INTERACTIVETOOLSFRAMEWORK_API virtual void OnPropertyModified(UObject* PropertySet, FProperty* Property) override;
 	
 	// IClickDragBehaviorTarget implementation
-	virtual FInputRayHit CanBeginClickDragSequence(const FInputDeviceRay& PressPos) override;
-	virtual void OnClickPress(const FInputDeviceRay& PressPos) override;
-	virtual void OnClickDrag(const FInputDeviceRay& DragPos) override;
+	INTERACTIVETOOLSFRAMEWORK_API virtual FInputRayHit CanBeginClickDragSequence(const FInputDeviceRay& PressPos) override;
+	INTERACTIVETOOLSFRAMEWORK_API virtual void OnClickPress(const FInputDeviceRay& PressPos) override;
+	INTERACTIVETOOLSFRAMEWORK_API virtual void OnClickDrag(const FInputDeviceRay& DragPos) override;
 	
 	// UMeshSurfacePointTool implementation
-	virtual void OnBeginDrag(const FRay& Ray) override;
-	virtual void OnUpdateDrag(const FRay& Ray) override;
-	virtual void OnEndDrag(const FRay& Ray) override;
-	virtual bool OnUpdateHover(const FInputDeviceRay& DevicePos) override;
+	INTERACTIVETOOLSFRAMEWORK_API virtual void OnBeginDrag(const FRay& Ray) override;
+	INTERACTIVETOOLSFRAMEWORK_API virtual void OnUpdateDrag(const FRay& Ray) override;
+	INTERACTIVETOOLSFRAMEWORK_API virtual void OnEndDrag(const FRay& Ray) override;
+	INTERACTIVETOOLSFRAMEWORK_API virtual bool OnUpdateHover(const FInputDeviceRay& DevicePos) override;
 
 public:
 
@@ -197,19 +197,19 @@ public:
 
 public:
 
-	virtual void IncreaseBrushSizeAction();
-	virtual void DecreaseBrushSizeAction();
-	virtual void IncreaseBrushStrengthAction();
-	virtual void DecreaseBrushStrengthAction();
-	virtual void IncreaseBrushFalloffAction();
-	virtual void DecreaseBrushFalloffAction();
+	INTERACTIVETOOLSFRAMEWORK_API virtual void IncreaseBrushSizeAction();
+	INTERACTIVETOOLSFRAMEWORK_API virtual void DecreaseBrushSizeAction();
+	INTERACTIVETOOLSFRAMEWORK_API virtual void IncreaseBrushStrengthAction();
+	INTERACTIVETOOLSFRAMEWORK_API virtual void DecreaseBrushStrengthAction();
+	INTERACTIVETOOLSFRAMEWORK_API virtual void IncreaseBrushFalloffAction();
+	INTERACTIVETOOLSFRAMEWORK_API virtual void DecreaseBrushFalloffAction();
 
 	virtual bool IsInBrushStroke() const { return bInBrushStroke; }
 
 	virtual double GetCurrentBrushRadius() const { return CurrentBrushRadius; }
 	virtual double GetCurrentBrushRadiusLocal() const { return CurrentBrushRadius * WorldToLocalScale; }
 
-	void SetBrushEnabled(bool bIsEnabled);
+	INTERACTIVETOOLSFRAMEWORK_API void SetBrushEnabled(bool bIsEnabled);
 
 	// return false to disable the hotkey for adjusting the brush radius and strength in the viewport
 	virtual bool SupportsBrushAdjustmentInput() { return true; }; 
@@ -224,7 +224,7 @@ protected:
 protected:
 	TInterval<float> BrushRelativeSizeRange;
 	double CurrentBrushRadius;
-	void RecalculateBrushRadius();
+	INTERACTIVETOOLSFRAMEWORK_API void RecalculateBrushRadius();
 
 	UPROPERTY()
 	TSoftClassPtr<UBrushBaseProperties> PropertyClass;
@@ -237,9 +237,9 @@ protected:
 	UPROPERTY()
 	TObjectPtr<UBrushStampIndicator> BrushStampIndicator;
 
-	virtual void SetupBrushStampIndicator();
-	virtual void UpdateBrushStampIndicator();
-	virtual void ShutdownBrushStampIndicator();
+	INTERACTIVETOOLSFRAMEWORK_API virtual void SetupBrushStampIndicator();
+	INTERACTIVETOOLSFRAMEWORK_API virtual void UpdateBrushStampIndicator();
+	INTERACTIVETOOLSFRAMEWORK_API virtual void ShutdownBrushStampIndicator();
 
 	// adjusts size and strength properties while holding hotkey during click/drag
 	TWeakObjectPtr<UBrushAdjusterInputBehavior> BrushAdjusterBehavior;
diff --git a/Engine/Source/Runtime/InteractiveToolsFramework/Public/BaseTools/ClickDragTool.h b/Engine/Source/Runtime/InteractiveToolsFramework/Public/BaseTools/ClickDragTool.h
--- a/Engine/Source/Runtime/InteractiveToolsFramework/Public/BaseTools/ClickDragTool.h
+++ b/Engine/Source/Runtime/InteractiveToolsFramework/Public/BaseTools/ClickDragTool.h
@@ -20,14 +20,14 @@ struct FToolBuilderState;
 /**
  * Builder for UClickDragTool
  */
-UCLASS()
-class INTERACTIVETOOLSFRAMEWORK_API UClickDragToolBuilder : public UInteractiveToolBuilder
+UCLASS(MinimalAPI)
+class UClickDragToolBuilder : public UInteractiveToolBuilder
 {
 	GENERATED_BODY()
 
 public:
-	virtual bool CanBuildTool(const FToolBuilderState& SceneState) const override;
-	virtual UInteractiveTool* BuildTool(const FToolBuilderState& SceneState) const override;
+	INTERACTIVETOOLSFRAMEWORK_API virtual bool CanBuildTool(const FToolBuilderState& SceneState) const override;
+	INTERACTIVETOOLSFRAMEWORK_API virtual UInteractiveTool* BuildTool(const FToolBuilderState& SceneState) const override;
 };
 
 
@@ -37,8 +37,8 @@ public:
  * implement basic click-drag type Tools. If you want to do more advanced things, 
  * like handle modifier buttons/keys, you will need to implement IClickDragBehaviorTarget yourself
  */
-UCLASS()
-class INTERACTIVETOOLSFRAMEWORK_API UClickDragTool : public UInteractiveTool, public IClickDragBehaviorTarget
+UCLASS(MinimalAPI)
+class UClickDragTool : public UInteractiveTool, public IClickDragBehaviorTarget
 {
 	GENERATED_BODY()
 
@@ -47,7 +47,7 @@ public:
 	/**
 	 * Register default primary-button-click InputBehaviors
 	 */
-	virtual void Setup() override;
+	INTERACTIVETOOLSFRAMEWORK_API virtual void Setup() override;
 
 
 	//
@@ -59,30 +59,30 @@ public:
 	 * @param ClickPos device position/ray at click point
 	 * @return true if target wants to begin sequence
 	 */
-	virtual FInputRayHit CanBeginClickDragSequence(const FInputDeviceRay& PressPos) override;
+	INTERACTIVETOOLSFRAMEWORK_API virtual FInputRayHit CanBeginClickDragSequence(const FInputDeviceRay& PressPos) override;
 
 	/**
 	 * Notify Target that click press ocurred
 	 * @param ClickPos device position/ray at click point
 	 */
-	virtual void OnClickPress(const FInputDeviceRay& PressPos) override;
+	INTERACTIVETOOLSFRAMEWORK_API virtual void OnClickPress(const FInputDeviceRay& PressPos) override;
 
 	/**
 	 * Notify Target that input position has changed
 	 * @param ClickPos device position/ray at click point
 	 */
-	virtual void OnClickDrag(const FInputDeviceRay& DragPos) override;
+	INTERACTIVETOOLSFRAMEWORK_API virtual void OnClickDrag(const FInputDeviceRay& DragPos) override;
 
 	/**
 	 * Notify Target that click release occurred
 	 * @param ClickPos device position/ray at click point
 	 */
-	virtual void OnClickRelease(const FInputDeviceRay& ReleasePos) override;
+	INTERACTIVETOOLSFRAMEWORK_API virtual void OnClickRelease(const FInputDeviceRay& ReleasePos) override;
 
 	/**
 	 * Notify Target that click-drag sequence has been explicitly terminated (eg by escape key)
 	 */
-	virtual void OnTerminateDragSequence() override;
+	INTERACTIVETOOLSFRAMEWORK_API virtual void OnTerminateDragSequence() override;
 
 };
 
diff --git a/Engine/Source/Runtime/InteractiveToolsFramework/Public/BaseTools/MeshSurfacePointTool.h b/Engine/Source/Runtime/InteractiveToolsFramework/Public/BaseTools/MeshSurfacePointTool.h
--- a/Engine/Source/Runtime/InteractiveToolsFramework/Public/BaseTools/MeshSurfacePointTool.h
+++ b/Engine/Source/Runtime/InteractiveToolsFramework/Public/BaseTools/MeshSurfacePointTool.h
@@ -16,7 +16,7 @@ class UMeshSurfacePointTool;
 
 // This is a temporary interface to provide stylus pressure, currently necessary due to limitations
 // in the stylus plugin architecture. Should be removed once InputState/InputBehavior can support stylus events
-class INTERACTIVETOOLSFRAMEWORK_API IToolStylusStateProviderAPI
+class IToolStylusStateProviderAPI
 {
 public:
 	virtual float GetCurrentPressure() const = 0;
@@ -25,8 +25,8 @@ public:
 /**
  * 
  */
-UCLASS()
-class INTERACTIVETOOLSFRAMEWORK_API UMeshSurfacePointToolBuilder : public UInteractiveToolWithToolTargetsBuilder
+UCLASS(MinimalAPI)
+class UMeshSurfacePointToolBuilder : public UInteractiveToolWithToolTargetsBuilder
 {
 	GENERATED_BODY()
 
@@ -34,19 +34,19 @@ public:
 	IToolStylusStateProviderAPI* StylusAPI = nullptr;
 
 	/** @return true if a single mesh source can be found in the active selection */
-	virtual bool CanBuildTool(const FToolBuilderState& SceneState) const override;
+	INTERACTIVETOOLSFRAMEWORK_API virtual bool CanBuildTool(const FToolBuilderState& SceneState) const override;
 
 	/** @return new Tool instance initialized with selected mesh source */
-	virtual UInteractiveTool* BuildTool(const FToolBuilderState& SceneState) const override;
+	INTERACTIVETOOLSFRAMEWORK_API virtual UInteractiveTool* BuildTool(const FToolBuilderState& SceneState) const override;
 
 	/** @return new Tool instance. Override this in subclasses to build a different Tool class type */
-	virtual UMeshSurfacePointTool* CreateNewTool(const FToolBuilderState& SceneState) const;
+	INTERACTIVETOOLSFRAMEWORK_API virtual UMeshSurfacePointTool* CreateNewTool(const FToolBuilderState& SceneState) const;
 
 	/** Called by BuildTool to configure the Tool with the input MeshSource based on the SceneState */
-	virtual void InitializeNewTool(UMeshSurfacePointTool* Tool, const FToolBuilderState& SceneState) const;
+	INTERACTIVETOOLSFRAMEWORK_API virtual void InitializeNewTool(UMeshSurfacePointTool* Tool, const FToolBuilderState& SceneState) const;
 
 protected:
-	virtual const FToolTargetTypeRequirements& GetTargetRequirements() const override;
+	INTERACTIVETOOLSFRAMEWORK_API virtual const FToolTargetTypeRequirements& GetTargetRequirements() const override;
 };
 
 
@@ -59,8 +59,8 @@ protected:
  * Subclasses override the OnBeginDrag/OnUpdateDrag/OnEndDrag and OnUpdateHover functions
  * to implement custom behavior.
  */
-UCLASS()
-class INTERACTIVETOOLSFRAMEWORK_API UMeshSurfacePointTool : public USingleSelectionTool, public IClickDragBehaviorTarget, public IHoverBehaviorTarget
+UCLASS(MinimalAPI)
+class UMeshSurfacePointTool : public USingleSelectionTool, public IClickDragBehaviorTarget, public IHoverBehaviorTarget
 {
 	GENERATED_BODY()
 
@@ -68,74 +68,74 @@ public:
 	// UInteractiveTool API Implementation
 
 	/** Register InputBehaviors, etc */
-	virtual void Setup() override;
+	INTERACTIVETOOLSFRAMEWORK_API virtual void Setup() override;
 
 	/** Set current stlyus API source */
-	virtual void SetStylusAPI(IToolStylusStateProviderAPI* StylusAPI);
+	INTERACTIVETOOLSFRAMEWORK_API virtual void SetStylusAPI(IToolStylusStateProviderAPI* StylusAPI);
 
 	// UMeshSurfacePointTool API
 
 	/**
 	 * @return true if the target MeshSource is hit by the Ray
 	 */
-	virtual bool HitTest(const FRay& Ray, FHitResult& OutHit);
+	INTERACTIVETOOLSFRAMEWORK_API virtual bool HitTest(const FRay& Ray, FHitResult& OutHit);
 
 	/**
 	 * This function is called when the user begins a click-drag-release interaction
 	 */
-	virtual void OnBeginDrag(const FRay& Ray);
+	INTERACTIVETOOLSFRAMEWORK_API virtual void OnBeginDrag(const FRay& Ray);
 
 	/**
 	 * This function is called each frame that the user is in a click-drag-release interaction
 	 */
-	virtual void OnUpdateDrag(const FRay& Ray);
+	INTERACTIVETOOLSFRAMEWORK_API virtual void OnUpdateDrag(const FRay& Ray);
 
 	/**
 	 * This function is called when the user releases the button driving a click-drag-release interaction
 	 */
-	virtual void OnEndDrag(const FRay& Ray);
+	INTERACTIVETOOLSFRAMEWORK_API virtual void OnEndDrag(const FRay& Ray);
 
 
 
 
 	// IClickDragBehaviorTarget implementation
-	virtual FInputRayHit CanBeginClickDragSequence(const FInputDeviceRay& PressPos) override;
-	virtual void OnClickPress(const FInputDeviceRay& PressPos) override;
-	virtual void OnClickDrag(const FInputDeviceRay& DragPos) override;
-	virtual void OnClickRelease(const FInputDeviceRay& ReleasePos) override;
-	virtual void OnTerminateDragSequence() override;
-	virtual void OnUpdateModifierState(int ModifierID, bool bIsOn) override;
+	INTERACTIVETOOLSFRAMEWORK_API virtual FInputRayHit CanBeginClickDragSequence(const FInputDeviceRay& PressPos) override;
+	INTERACTIVETOOLSFRAMEWORK_API virtual void OnClickPress(const FInputDeviceRay& PressPos) override;
+	INTERACTIVETOOLSFRAMEWORK_API virtual void OnClickDrag(const FInputDeviceRay& DragPos) override;
+	INTERACTIVETOOLSFRAMEWORK_API virtual void OnClickRelease(const FInputDeviceRay& ReleasePos) override;
+	INTERACTIVETOOLSFRAMEWORK_API virtual void OnTerminateDragSequence() override;
+	INTERACTIVETOOLSFRAMEWORK_API virtual void OnUpdateModifierState(int ModifierID, bool bIsOn) override;
 
 
 	// IHoverBehaviorTarget implementation
 	
-	virtual FInputRayHit BeginHoverSequenceHitTest(const FInputDeviceRay& PressPos) override;
+	INTERACTIVETOOLSFRAMEWORK_API virtual FInputRayHit BeginHoverSequenceHitTest(const FInputDeviceRay& PressPos) override;
 	virtual void OnBeginHover(const FInputDeviceRay& DevicePos) override {}
 	virtual bool OnUpdateHover(const FInputDeviceRay& DevicePos) override { return true; }
 	virtual void OnEndHover() override {}
 
 
 	// IInteractiveToolCameraFocusAPI implementation
-	virtual bool GetWorldSpaceFocusPoint(const FRay& WorldRay, FVector& PointOut) override;
+	INTERACTIVETOOLSFRAMEWORK_API virtual bool GetWorldSpaceFocusPoint(const FRay& WorldRay, FVector& PointOut) override;
 
 
 	/** Called by registered InputBehaviors to set the state of the "shift" button (or device equivalent) */
-	virtual void SetShiftToggle(bool bShiftDown);
+	INTERACTIVETOOLSFRAMEWORK_API virtual void SetShiftToggle(bool bShiftDown);
 
 	/** @return current state of the shift toggle */
 	virtual bool GetShiftToggle() const { return bShiftToggle; }
 
 	/** Called by registered InputBehaviors to set the state of the "ctrl" button (or device equivalent) */
-	virtual void SetCtrlToggle(bool bCtrlDown);
+	INTERACTIVETOOLSFRAMEWORK_API virtual void SetCtrlToggle(bool bCtrlDown);
 
 	/** @return current state of the ctrl toggle */
 	virtual bool GetCtrlToggle() const { return bCtrlToggle; }
 
 	/** @return current input device pressure in range 0-1 */
-	virtual float GetCurrentDevicePressure() const;
+	INTERACTIVETOOLSFRAMEWORK_API virtual float GetCurrentDevicePressure() const;
 
-	void SetWorld(UWorld* World);
-	UWorld* GetTargetWorld();
+	INTERACTIVETOOLSFRAMEWORK_API void SetWorld(UWorld* World);
+	INTERACTIVETOOLSFRAMEWORK_API UWorld* GetTargetWorld();
 
 protected:
 	/** Current state of the shift modifier toggle */
diff --git a/Engine/Source/Runtime/InteractiveToolsFramework/Public/BaseTools/SingleClickTool.h b/Engine/Source/Runtime/InteractiveToolsFramework/Public/BaseTools/SingleClickTool.h
--- a/Engine/Source/Runtime/InteractiveToolsFramework/Public/BaseTools/SingleClickTool.h
+++ b/Engine/Source/Runtime/InteractiveToolsFramework/Public/BaseTools/SingleClickTool.h
@@ -21,14 +21,14 @@ struct FToolBuilderState;
 /**
  * Builder for USingleClickTool
  */
-UCLASS()
-class INTERACTIVETOOLSFRAMEWORK_API USingleClickToolBuilder : public UInteractiveToolBuilder
+UCLASS(MinimalAPI)
+class USingleClickToolBuilder : public UInteractiveToolBuilder
 {
 	GENERATED_BODY()
 
 public:
-	virtual bool CanBuildTool(const FToolBuilderState& SceneState) const override;
-	virtual UInteractiveTool* BuildTool(const FToolBuilderState& SceneState) const override;
+	INTERACTIVETOOLSFRAMEWORK_API virtual bool CanBuildTool(const FToolBuilderState& SceneState) const override;
+	INTERACTIVETOOLSFRAMEWORK_API virtual UInteractiveTool* BuildTool(const FToolBuilderState& SceneState) const override;
 };
 
 
@@ -44,8 +44,8 @@ public:
  * The function ::OnClicked() implements the action that will occur when a click happens.
  * You must override this to implement any kind of useful behavior.
  */
-UCLASS()
-class INTERACTIVETOOLSFRAMEWORK_API USingleClickTool : public UInteractiveTool, public IClickBehaviorTarget
+UCLASS(MinimalAPI)
+class USingleClickTool : public UInteractiveTool, public IClickBehaviorTarget
 {
 	GENERATED_BODY()
 
@@ -54,17 +54,17 @@ public:
 	/**
 	 * Register default primary-button-click InputBehaviors
 	 */
-	virtual void Setup() override;
+	INTERACTIVETOOLSFRAMEWORK_API virtual void Setup() override;
 
 	/**
 	 * Test if the Target is hit at this 2D position / 3D ray
 	 */
-	virtual FInputRayHit IsHitByClick(const FInputDeviceRay& ClickPos);
+	INTERACTIVETOOLSFRAMEWORK_API virtual FInputRayHit IsHitByClick(const FInputDeviceRay& ClickPos);
 
 	/**
 	 * Click the Target at this 2D position / 3D ray. Default behavior is to print debug string.
 	 */
-	virtual void OnClicked(const FInputDeviceRay& ClickPos);
+	INTERACTIVETOOLSFRAMEWORK_API virtual void OnClicked(const FInputDeviceRay& ClickPos);
 
 };
 
diff --git a/Engine/Source/Runtime/InteractiveToolsFramework/Public/Changes/TransformChange.h b/Engine/Source/Runtime/InteractiveToolsFramework/Public/Changes/TransformChange.h
--- a/Engine/Source/Runtime/InteractiveToolsFramework/Public/Changes/TransformChange.h
+++ b/Engine/Source/Runtime/InteractiveToolsFramework/Public/Changes/TransformChange.h
@@ -14,11 +14,11 @@ class UObject;
 /**
  * FComponentWorldTransformChange represents an undoable change to the world transform of a USceneComponent.
  */
-class INTERACTIVETOOLSFRAMEWORK_API FComponentWorldTransformChange : public FToolCommandChange
+class FComponentWorldTransformChange : public FToolCommandChange
 {
 public:
-	FComponentWorldTransformChange();
-	FComponentWorldTransformChange(const FTransform& From, const FTransform& To);
+	INTERACTIVETOOLSFRAMEWORK_API FComponentWorldTransformChange();
+	INTERACTIVETOOLSFRAMEWORK_API FComponentWorldTransformChange(const FTransform& From, const FTransform& To);
 
 	FTransform FromWorldTransform;
 	FTransform ToWorldTransform;
@@ -27,13 +27,13 @@ public:
 	TFunction<void(FComponentWorldTransformChange*, UObject*, bool)> OnChangeAppliedFunc;
 
 	/** Makes the change to the object */
-	virtual void Apply(UObject* Object) override;
+	INTERACTIVETOOLSFRAMEWORK_API virtual void Apply(UObject* Object) override;
 
 	/** Reverts change to the object */
-	virtual void Revert(UObject* Object) override;
+	INTERACTIVETOOLSFRAMEWORK_API virtual void Revert(UObject* Object) override;
 
 	/** Describes this change (for debugging) */
-	virtual FString ToString() const override;
+	INTERACTIVETOOLSFRAMEWORK_API virtual FString ToString() const override;
 };
 
 
diff --git a/Engine/Source/Runtime/InteractiveToolsFramework/Public/ComponentSourceInterfaces.h b/Engine/Source/Runtime/InteractiveToolsFramework/Public/ComponentSourceInterfaces.h
--- a/Engine/Source/Runtime/InteractiveToolsFramework/Public/ComponentSourceInterfaces.h
+++ b/Engine/Source/Runtime/InteractiveToolsFramework/Public/ComponentSourceInterfaces.h
@@ -26,7 +26,7 @@ struct FComponentMaterialSet;
  *
  * (Conceivably this doesn't have to be backed by a Component, but most usage will assume there is an Actor)
  */
-class INTERACTIVETOOLSFRAMEWORK_API FPrimitiveComponentTarget
+class FPrimitiveComponentTarget
 {
 public:
 	virtual ~FPrimitiveComponentTarget(){}
@@ -37,36 +37,36 @@ public:
 	FPrimitiveComponentTarget( UPrimitiveComponent* Component ): Component( Component ){}
 
 	/** @return true if component target is still valid. May become invalid for various reasons (eg Component was deleted out from under us) */
-	virtual bool IsValid() const;
+	INTERACTIVETOOLSFRAMEWORK_API virtual bool IsValid() const;
 
 	/** @return the Actor that owns this Component */
-	AActor* GetOwnerActor() const;
+	INTERACTIVETOOLSFRAMEWORK_API AActor* GetOwnerActor() const;
 
 	/** @return the Component this is a Source for */
-	UPrimitiveComponent* GetOwnerComponent() const;
+	INTERACTIVETOOLSFRAMEWORK_API UPrimitiveComponent* GetOwnerComponent() const;
 
 	/** @return number of material indices in use by this Component */
-	int32 GetNumMaterials() const;
+	INTERACTIVETOOLSFRAMEWORK_API int32 GetNumMaterials() const;
 
 	/**
 	 * Get pointer to a Material provided by this Source
 	 * @param MaterialIndex index of the material
 	 * @return MaterialInterface pointer, or null if MaterialIndex is invalid
 	 */
-	UMaterialInterface* GetMaterial(int32 MaterialIndex) const;
+	INTERACTIVETOOLSFRAMEWORK_API UMaterialInterface* GetMaterial(int32 MaterialIndex) const;
 
 	/**
 	 * Get material set provided by this source
 	 * @param MaterialSetOut returned material set
 	 * @param bAssetMaterials if an underlying asset exists, return the Asset-level material assignment instead of the component materials
 	 */
-	virtual void GetMaterialSet(FComponentMaterialSet& MaterialSetOut, bool bAssetMaterials = false) const;
+	INTERACTIVETOOLSFRAMEWORK_API virtual void GetMaterialSet(FComponentMaterialSet& MaterialSetOut, bool bAssetMaterials = false) const;
 
 	/**
 	 * @return the transform on this component
 	 * @todo Do we need to return a list of transforms here?
 	 */
-	FTransform GetWorldTransform() const;
+	INTERACTIVETOOLSFRAMEWORK_API FTransform GetWorldTransform() const;
 
 	/**
 	 * Compute ray intersection with the MeshDescription this Source is providing
@@ -74,13 +74,13 @@ public:
 	 * @param OutHit hit test data
 	 * @return true if ray intersected Component
 	 */
-	bool HitTest(const FRay& WorldRay, FHitResult& OutHit) const;
+	INTERACTIVETOOLSFRAMEWORK_API bool HitTest(const FRay& WorldRay, FHitResult& OutHit) const;
 
 	/**
 	 * Set the visibility of the Component associated with this Source (ie to hide during Tool usage)
 	 * @param bVisible desired visibility
 	 */
-	void SetOwnerVisibility(bool bVisible) const;
+	INTERACTIVETOOLSFRAMEWORK_API void SetOwnerVisibility(bool bVisible) const;
 
 
 	/**
@@ -95,7 +95,7 @@ public:
 	 * @param MaterialSet new list of materials
 	 * @param bApplyToAsset if true, materials of Asset are updated (if Asset exists), rather than Component
 	 */
-	virtual void CommitMaterialSetUpdate(const FComponentMaterialSet& MaterialSet, bool bApplyToAsset);
+	INTERACTIVETOOLSFRAMEWORK_API virtual void CommitMaterialSetUpdate(const FComponentMaterialSet& MaterialSet, bool bApplyToAsset);
 
 
 	struct FCommitParams
@@ -110,7 +110,7 @@ public:
 };
 
 /** @deprecated Use tool target factories instead. */
-class INTERACTIVETOOLSFRAMEWORK_API FComponentTargetFactory
+class FComponentTargetFactory
 {
 public:
 	virtual ~FComponentTargetFactory(){}
diff --git a/Engine/Source/Runtime/InteractiveToolsFramework/Public/ContextObjectStore.h b/Engine/Source/Runtime/InteractiveToolsFramework/Public/ContextObjectStore.h
--- a/Engine/Source/Runtime/InteractiveToolsFramework/Public/ContextObjectStore.h
+++ b/Engine/Source/Runtime/InteractiveToolsFramework/Public/ContextObjectStore.h
@@ -20,8 +20,8 @@ class UClass;
   *   - A tool may allow extra actions if it has access to a particular API object in the context store.
   *   - A tool may choose to initialize itself differently based on the presence of a selection-holding data object in the context store.
  */
-UCLASS(Transient)
-class INTERACTIVETOOLSFRAMEWORK_API UContextObjectStore : public UObject
+UCLASS(Transient, MinimalAPI)
+class UContextObjectStore : public UObject
 {
 	GENERATED_BODY()
 
@@ -53,25 +53,25 @@ public:
 	 * Finds the first context object that derives from the given class.
 	 * @returns the found context object, or nullptr if none matches.
 	 */
-	UObject* FindContextByClass(UClass* InClass) const;
+	INTERACTIVETOOLSFRAMEWORK_API UObject* FindContextByClass(UClass* InClass) const;
 
 	/**
 	 * Adds a data object to the tool manager's set of shared data objects.
 	 * @returns true if the addition is successful.
 	 */
-	bool AddContextObject(UObject* InContextObject);
+	INTERACTIVETOOLSFRAMEWORK_API bool AddContextObject(UObject* InContextObject);
 
 	/**
 	 * Removes a data object from the tool manager's set of shared data objects.
 	 * @returns true if the removal is successful.
 	 */
-	bool RemoveContextObject(UObject* InContextObject);
+	INTERACTIVETOOLSFRAMEWORK_API bool RemoveContextObject(UObject* InContextObject);
 
 	/**
 	 * Removes any data objects from the tool manager's set of shared data objects that are of type @param InClass
 	 * @returns true if any objects are removed.
 	 */
-	bool RemoveContextObjectsOfType(const UClass* InClass);
+	INTERACTIVETOOLSFRAMEWORK_API bool RemoveContextObjectsOfType(const UClass* InClass);
 
 	template <class TObjectType>
 	bool RemoveContextObjectsOfType()
@@ -82,7 +82,7 @@ public:
 	/**
 	 * Shuts down the context object store, releasing hold on any stored content objects.
 	 */
-	void Shutdown();
+	INTERACTIVETOOLSFRAMEWORK_API void Shutdown();
 
 protected:
 	UPROPERTY()
diff --git a/Engine/Source/Runtime/InteractiveToolsFramework/Public/InputBehavior.h b/Engine/Source/Runtime/InteractiveToolsFramework/Public/InputBehavior.h
--- a/Engine/Source/Runtime/InteractiveToolsFramework/Public/InputBehavior.h
+++ b/Engine/Source/Runtime/InteractiveToolsFramework/Public/InputBehavior.h
@@ -34,7 +34,7 @@ enum class EInputCaptureSide
  * is capturing for Left and Right separately). So FInputCaptureUpdate can optionally
  * return this structure, and we will pass it to the next UpdateCapture() call
  */
-struct INTERACTIVETOOLSFRAMEWORK_API FInputCaptureData
+struct FInputCaptureData
 {
 	/** Which side do we want to capture on */
 	EInputCaptureSide WhichSide;
@@ -69,7 +69,7 @@ class UInputBehavior;
  * UInputBehavior returns an FInputCaptureRequest from WantsCapture() to indicate
  * whether it wants to capture or ignore an input event
  */
-struct INTERACTIVETOOLSFRAMEWORK_API FInputCaptureRequest
+struct FInputCaptureRequest
 {
 	/** Which input behavior generated this request */
 	UInputBehavior* Source;
@@ -133,7 +133,7 @@ enum class EInputCaptureState
  * IInputBehavior returns an FInputCaptureUpdate from BeginCapture() and UpdateCapture(),
  * which indicates to the InputRouter what the Behavior would like to have happen.
  */
-struct INTERACTIVETOOLSFRAMEWORK_API FInputCaptureUpdate
+struct FInputCaptureUpdate
 {
 	/** Indicates what capture state the Behavior wants to transition to */
 	EInputCaptureState State;
@@ -190,7 +190,7 @@ struct INTERACTIVETOOLSFRAMEWORK_API FInputCaptureUpdate
  * Each UInputBehavior provides a priority that is used to help resolve situations
  * when multiple Behaviors want to capture based on the same input event
  */
-struct INTERACTIVETOOLSFRAMEWORK_API FInputCapturePriority
+struct FInputCapturePriority
 {
 	static constexpr int DEFAULT_GIZMO_PRIORITY = 50;
 	static constexpr int DEFAULT_TOOL_PRIORITY = 100;
@@ -251,38 +251,38 @@ struct INTERACTIVETOOLSFRAMEWORK_API FInputCapturePriority
  *
  * Implementing interactions in this way allows the input handling to be separated from functionality.
  */
-UCLASS(Transient)
-class INTERACTIVETOOLSFRAMEWORK_API UInputBehavior : public UObject
+UCLASS(Transient, MinimalAPI)
+class UInputBehavior : public UObject
 {
 	GENERATED_BODY()
 
 public:
-	UInputBehavior();
+	INTERACTIVETOOLSFRAMEWORK_API UInputBehavior();
 
 	/** The priority is used to resolve situations where multiple behaviors want the same capture */
-	virtual FInputCapturePriority GetPriority();
+	INTERACTIVETOOLSFRAMEWORK_API virtual FInputCapturePriority GetPriority();
 
 	/** Configure the default priority of an instance of this behavior */
-	virtual void SetDefaultPriority(const FInputCapturePriority& Priority);
+	INTERACTIVETOOLSFRAMEWORK_API virtual void SetDefaultPriority(const FInputCapturePriority& Priority);
 
 
 	/** Which device types does this Behavior support */
-	virtual EInputDevices GetSupportedDevices();
+	INTERACTIVETOOLSFRAMEWORK_API virtual EInputDevices GetSupportedDevices();
 
 	/** Given the input state, does this Behavior want to begin capturing some input devices? */
-	virtual FInputCaptureRequest WantsCapture(const FInputDeviceState& InputState);
+	INTERACTIVETOOLSFRAMEWORK_API virtual FInputCaptureRequest WantsCapture(const FInputDeviceState& InputState);
 
 	/** Called after WantsCapture() returns a capture request that was accepted */
-	virtual FInputCaptureUpdate BeginCapture(const FInputDeviceState& InputState, EInputCaptureSide eSide);
+	INTERACTIVETOOLSFRAMEWORK_API virtual FInputCaptureUpdate BeginCapture(const FInputDeviceState& InputState, EInputCaptureSide eSide);
 
 	/** 
 	 * Called for each new input event during a capture sequence. Return Continue to keep
 	 * capturing, or End to finish capturing.
 	 */
-	virtual FInputCaptureUpdate UpdateCapture(const FInputDeviceState& InputState, const FInputCaptureData& CaptureData);
+	INTERACTIVETOOLSFRAMEWORK_API virtual FInputCaptureUpdate UpdateCapture(const FInputDeviceState& InputState, const FInputCaptureData& CaptureData);
 
 	/** If this is called, the Behavior has forcibly lost capture (eg due to app losing focus for example) and needs to clean up accordingly */
-	virtual void ForceEndCapture(const FInputCaptureData& CaptureData);
+	INTERACTIVETOOLSFRAMEWORK_API virtual void ForceEndCapture(const FInputCaptureData& CaptureData);
 
 
 	//
@@ -290,22 +290,22 @@ public:
 	//
 
 	/** return true if this Behavior supports hover (ie passive input events) */
-	virtual bool WantsHoverEvents();
+	INTERACTIVETOOLSFRAMEWORK_API virtual bool WantsHoverEvents();
 
 	/** Given the input state, does this Behavior want to begin capturing some input devices for hover */
-	virtual FInputCaptureRequest WantsHoverCapture(const FInputDeviceState& InputState);
+	INTERACTIVETOOLSFRAMEWORK_API virtual FInputCaptureRequest WantsHoverCapture(const FInputDeviceState& InputState);
 
 	/** Called after WantsHoverCapture() returns a capture request that was accepted */
-	virtual FInputCaptureUpdate BeginHoverCapture(const FInputDeviceState& InputState, EInputCaptureSide eSide);
+	INTERACTIVETOOLSFRAMEWORK_API virtual FInputCaptureUpdate BeginHoverCapture(const FInputDeviceState& InputState, EInputCaptureSide eSide);
 
 	/** Called on each new hover input event, ie if no other behavior is actively capturing input */
-	virtual FInputCaptureUpdate UpdateHoverCapture(const FInputDeviceState& InputState);
+	INTERACTIVETOOLSFRAMEWORK_API virtual FInputCaptureUpdate UpdateHoverCapture(const FInputDeviceState& InputState);
 
 	/** If a different hover capture begins, focus is lost, a tool starts, etc, any active hover visualization needs to terminate */
-	virtual void EndHoverCapture();
+	INTERACTIVETOOLSFRAMEWORK_API virtual void EndHoverCapture();
 
 
 protected:
 	/** priority returned by GetPriority() */
 	FInputCapturePriority DefaultPriority;
-};
\ No newline at end of file
+};
diff --git a/Engine/Source/Runtime/InteractiveToolsFramework/Public/InputBehaviorSet.h b/Engine/Source/Runtime/InteractiveToolsFramework/Public/InputBehaviorSet.h
--- a/Engine/Source/Runtime/InteractiveToolsFramework/Public/InputBehaviorSet.h
+++ b/Engine/Source/Runtime/InteractiveToolsFramework/Public/InputBehaviorSet.h
@@ -22,7 +22,7 @@
  * This is an internal structure used by UInputBehaviorSet. 
  */
 USTRUCT()
-struct INTERACTIVETOOLSFRAMEWORK_API FBehaviorInfo
+struct FBehaviorInfo
 {
 	GENERATED_BODY()
 	
@@ -54,14 +54,14 @@ struct INTERACTIVETOOLSFRAMEWORK_API FBehaviorInfo
  * behaviors to be managed together. For example one UInputBehaviorSet can be merged into 
  * another and removed later.
  */
-UCLASS(Transient)
-class INTERACTIVETOOLSFRAMEWORK_API UInputBehaviorSet : public UObject
+UCLASS(Transient, MinimalAPI)
+class UInputBehaviorSet : public UObject
 {
 	GENERATED_BODY()
 
 public:
-	UInputBehaviorSet();
-	virtual ~UInputBehaviorSet();
+	INTERACTIVETOOLSFRAMEWORK_API UInputBehaviorSet();
+	INTERACTIVETOOLSFRAMEWORK_API virtual ~UInputBehaviorSet();
 
 	//
 	// Set Management
@@ -69,7 +69,7 @@ public:
 
 
 	/** @return true if there are no Behaviors in set */
-	virtual bool IsEmpty() const;
+	INTERACTIVETOOLSFRAMEWORK_API virtual bool IsEmpty() const;
 
 	/** 
 	 * Add a Behavior to the set
@@ -77,7 +77,7 @@ public:
 	 * @param Source pointer to owning object, used only to identify Behavior later
 	 * @param GroupName string identifier for this Behavior or group of Behaviors
 	 */
-	virtual void Add(UInputBehavior* Behavior, void* Source = nullptr, const FString& GroupName = "");
+	INTERACTIVETOOLSFRAMEWORK_API virtual void Add(UInputBehavior* Behavior, void* Source = nullptr, const FString& GroupName = "");
 
 	/**
 	 * Merge another BehaviorSet into this Set
@@ -85,30 +85,30 @@ public:
 	 * @param Source pointer to owning object, used only to identify Behavior later. If nullptr, source is copied from other Set.
 	 * @param NewGroupName string identifier for this Behavior or group of Behaviors. If empty string, group is copied from other Set.
 	 */
-	virtual void Add(const UInputBehaviorSet* OtherSet, void* NewSource = nullptr, const FString& NewGroupName = "");
+	INTERACTIVETOOLSFRAMEWORK_API virtual void Add(const UInputBehaviorSet* OtherSet, void* NewSource = nullptr, const FString& NewGroupName = "");
 
 	/**
 	 * Remove a Behavior from the Set
 	 * @param Behavior Behavior to remove
 	 * @return true if Behavior was found and removed
 	 */
-	virtual bool Remove(UInputBehavior* Behavior);
+	INTERACTIVETOOLSFRAMEWORK_API virtual bool Remove(UInputBehavior* Behavior);
 	/**
 	 * Remove a group of Behaviors from the Set
 	 * @param GroupName name of group, all Behaviors that were added with this GroupName are removed
 	 * @return true if any Behaviors were found and removed
 	 */
-	virtual bool RemoveByGroup(const FString& GroupName);
+	INTERACTIVETOOLSFRAMEWORK_API virtual bool RemoveByGroup(const FString& GroupName);
 	/**
 	 * Remove a group of Behaviors from the Set
 	 * @param Source source object pointer, all Behaviors that were added with this Source pointer are removed.
 	 * @return true if any Behaviors were found and removed
 	 */
-	virtual bool RemoveBySource(void* Source);
+	INTERACTIVETOOLSFRAMEWORK_API virtual bool RemoveBySource(void* Source);
 	/**
 	 * Remove all Behaviors from the set
 	 */
-	virtual void RemoveAll();
+	INTERACTIVETOOLSFRAMEWORK_API virtual void RemoveAll();
 
 
 	//
@@ -121,14 +121,14 @@ public:
 	 * @param InputState current input device state information
 	 * @param ResultOut returned set of non-Ignoring FInputCaptureRequests returned by Behaviors
 	 */
-	virtual void CollectWantsCapture(const FInputDeviceState& InputState, TArray<FInputCaptureRequest>& ResultOut);
+	INTERACTIVETOOLSFRAMEWORK_API virtual void CollectWantsCapture(const FInputDeviceState& InputState, TArray<FInputCaptureRequest>& ResultOut);
 
 	/**
 	 * Call UInputBehavior::WantsCapture() on each valid Behavior and collect up the requests that indicated a Capture was desired.
 	 * @param InputState current input device state information
 	 * @param ResultOut returned set of non-Ignoring FInputCaptureRequests returned by Behaviors
 	 */
-	virtual void CollectWantsHoverCapture(const FInputDeviceState& InputState, TArray<FInputCaptureRequest>& ResultOut);
+	INTERACTIVETOOLSFRAMEWORK_API virtual void CollectWantsHoverCapture(const FInputDeviceState& InputState, TArray<FInputCaptureRequest>& ResultOut);
 
 
 
@@ -141,7 +141,7 @@ protected:
 	/**
 	 * called internally when Behaviors list is updated, to re-sort by Priority, etc (@todo: and emit change event)
 	 */
-	virtual void BehaviorsModified();
+	INTERACTIVETOOLSFRAMEWORK_API virtual void BehaviorsModified();
 
 
 	/** @return true if Behavior supports InputState.InputDevice */
@@ -186,8 +186,8 @@ public:
  * a tool to supply a behavior source different from the one it is implementing itself. Useful,
  * for instance, when a tool wants to supply different behaviors to separate input routers.
  */
-UCLASS(Transient)
-class INTERACTIVETOOLSFRAMEWORK_API ULocalInputBehaviorSource : public UObject, public IInputBehaviorSource
+UCLASS(Transient, MinimalAPI)
+class ULocalInputBehaviorSource : public UObject, public IInputBehaviorSource
 {
 	GENERATED_BODY()
 
diff --git a/Engine/Source/Runtime/InteractiveToolsFramework/Public/InputRouter.h b/Engine/Source/Runtime/InteractiveToolsFramework/Public/InputRouter.h
--- a/Engine/Source/Runtime/InteractiveToolsFramework/Public/InputRouter.h
+++ b/Engine/Source/Runtime/InteractiveToolsFramework/Public/InputRouter.h
@@ -34,45 +34,45 @@ struct FInputDeviceState;
  * hover events to InputBehaviors that also support it.
  *
  */
-UCLASS(Transient)
-class INTERACTIVETOOLSFRAMEWORK_API UInputRouter : public UObject
+UCLASS(Transient, MinimalAPI)
+class UInputRouter : public UObject
 {
 	GENERATED_BODY()
 
 protected:
 	friend class UInteractiveToolsContext;		// to call Initialize/Shutdown
 
-	UInputRouter();
+	INTERACTIVETOOLSFRAMEWORK_API UInputRouter();
 
 	/** Initialize the InputRouter with the necessary Context-level state. UInteractiveToolsContext calls this, you should not. */
-	virtual void Initialize(IToolsContextTransactionsAPI* TransactionsAPI);
+	INTERACTIVETOOLSFRAMEWORK_API virtual void Initialize(IToolsContextTransactionsAPI* TransactionsAPI);
 
 	/** Shutdown the InputRouter. Called by UInteractiveToolsContext. */
-	virtual void Shutdown();
+	INTERACTIVETOOLSFRAMEWORK_API virtual void Shutdown();
 
 public:
 	/** Add a new behavior Source. Behaviors from this source will be added to the active behavior set. */
-	virtual void RegisterSource(IInputBehaviorSource* Source);
+	INTERACTIVETOOLSFRAMEWORK_API virtual void RegisterSource(IInputBehaviorSource* Source);
 
 	/** Remove Behaviors from this Source from the active set */
-	virtual void DeregisterSource(IInputBehaviorSource* Source);
+	INTERACTIVETOOLSFRAMEWORK_API virtual void DeregisterSource(IInputBehaviorSource* Source);
 
 	/** Insert a new input event which is used to check for new captures, or forwarded to active capture */
-	virtual bool PostInputEvent(const FInputDeviceState& Input);
+	INTERACTIVETOOLSFRAMEWORK_API virtual bool PostInputEvent(const FInputDeviceState& Input);
 
 	/** Returns true if there is an active mouse capture */
-	virtual bool HasActiveMouseCapture() const;
+	INTERACTIVETOOLSFRAMEWORK_API virtual bool HasActiveMouseCapture() const;
 
 	// TODO: other capture queries
 
 	/** Insert a new hover input event which is forwarded to all hover-enabled Behaviors */
-	virtual void PostHoverInputEvent(const FInputDeviceState& Input);
+	INTERACTIVETOOLSFRAMEWORK_API virtual void PostHoverInputEvent(const FInputDeviceState& Input);
 
 	/** If this Behavior is capturing, call ForceEndCapture() to notify that we are taking capture away */
-	virtual void ForceTerminateSource(IInputBehaviorSource* Source);
+	INTERACTIVETOOLSFRAMEWORK_API virtual void ForceTerminateSource(IInputBehaviorSource* Source);
 
 	/** Terminate any active captures and end all hovers */
-	virtual void ForceTerminateAll();
+	INTERACTIVETOOLSFRAMEWORK_API virtual void ForceTerminateAll();
 
 
 public:
@@ -103,13 +103,13 @@ protected:
 	void* ActiveRightCaptureOwner = nullptr;
 	FInputCaptureData ActiveRightCaptureData;
 
-	virtual void PostInputEvent_Keyboard(const FInputDeviceState& Input);
-	void CheckForKeyboardCaptures(const FInputDeviceState& Input);
-	void HandleCapturedKeyboardInput(const FInputDeviceState& Input);
+	INTERACTIVETOOLSFRAMEWORK_API virtual void PostInputEvent_Keyboard(const FInputDeviceState& Input);
+	INTERACTIVETOOLSFRAMEWORK_API void CheckForKeyboardCaptures(const FInputDeviceState& Input);
+	INTERACTIVETOOLSFRAMEWORK_API void HandleCapturedKeyboardInput(const FInputDeviceState& Input);
 
-	virtual void PostInputEvent_Mouse(const FInputDeviceState& Input);
-	void CheckForMouseCaptures(const FInputDeviceState& Input);
-	void HandleCapturedMouseInput(const FInputDeviceState& Input);
+	INTERACTIVETOOLSFRAMEWORK_API virtual void PostInputEvent_Mouse(const FInputDeviceState& Input);
+	INTERACTIVETOOLSFRAMEWORK_API void CheckForMouseCaptures(const FInputDeviceState& Input);
+	INTERACTIVETOOLSFRAMEWORK_API void HandleCapturedMouseInput(const FInputDeviceState& Input);
 
 	//
 	// Hover support
@@ -120,6 +120,6 @@ protected:
 
 	FInputDeviceState LastMouseInputState;
 
-	void TerminateHover(EInputCaptureSide Side);
-	bool ProcessMouseHover(const FInputDeviceState& Input);
-};
\ No newline at end of file
+	INTERACTIVETOOLSFRAMEWORK_API void TerminateHover(EInputCaptureSide Side);
+	INTERACTIVETOOLSFRAMEWORK_API bool ProcessMouseHover(const FInputDeviceState& Input);
+};
diff --git a/Engine/Source/Runtime/InteractiveToolsFramework/Public/InputState.h b/Engine/Source/Runtime/InteractiveToolsFramework/Public/InputState.h
--- a/Engine/Source/Runtime/InteractiveToolsFramework/Public/InputState.h
+++ b/Engine/Source/Runtime/InteractiveToolsFramework/Public/InputState.h
@@ -48,7 +48,7 @@ ENUM_CLASS_FLAGS(EInputDevices);
  * Generally this is intended to be returned as the result of a hit-test with a FInputDeviceRay 
  */
 USTRUCT(BlueprintType)
-struct INTERACTIVETOOLSFRAMEWORK_API FInputRayHit
+struct FInputRayHit
 {
 	GENERATED_BODY()
 
@@ -323,7 +323,7 @@ struct FInputDeviceState
  * have an associated 2D screen position.
  */
 USTRUCT(BlueprintType)
-struct INTERACTIVETOOLSFRAMEWORK_API FInputDeviceRay
+struct FInputDeviceRay
 {
 	GENERATED_BODY()
 
diff --git a/Engine/Source/Runtime/InteractiveToolsFramework/Public/InteractionMechanic.h b/Engine/Source/Runtime/InteractiveToolsFramework/Public/InteractionMechanic.h
--- a/Engine/Source/Runtime/InteractiveToolsFramework/Public/InteractionMechanic.h
+++ b/Engine/Source/Runtime/InteractiveToolsFramework/Public/InteractionMechanic.h
@@ -25,47 +25,47 @@ class UInteractiveToolPropertySet;
  * 
  * (This class is still a work in progress)
  */
-UCLASS(Transient)
-class INTERACTIVETOOLSFRAMEWORK_API UInteractionMechanic : public UObject
+UCLASS(Transient, MinimalAPI)
+class UInteractionMechanic : public UObject
 {
 	GENERATED_BODY()
 
 public:
-	UInteractionMechanic();
+	INTERACTIVETOOLSFRAMEWORK_API UInteractionMechanic();
 
 	/**
 	 * Called to initialize the InteractionMechanic
 	 */
-	virtual void Setup(UInteractiveTool* ParentTool);
+	INTERACTIVETOOLSFRAMEWORK_API virtual void Setup(UInteractiveTool* ParentTool);
 
 	/**
 	 * Called to clean up the InteractionMechanic
 	 */
-	virtual void Shutdown();
+	INTERACTIVETOOLSFRAMEWORK_API virtual void Shutdown();
 
 	/**
 	 * Allow the Mechanic to do any custom drawing (ie via PDI/RHI)
 	 * @param RenderAPI Abstraction that provides access to Rendering in the current ToolsContext
 	 */
-	virtual void Render(IToolsContextRenderAPI* RenderAPI);
+	INTERACTIVETOOLSFRAMEWORK_API virtual void Render(IToolsContextRenderAPI* RenderAPI);
 
 	/**
 	 * ALlow the Mechanic to Tick
 	 */
-	virtual void Tick(float DeltaTime);
+	INTERACTIVETOOLSFRAMEWORK_API virtual void Tick(float DeltaTime);
 
 protected:
 	TWeakObjectPtr<UInteractiveTool> ParentTool;
 
-	UInteractiveTool* GetParentTool() const;
+	INTERACTIVETOOLSFRAMEWORK_API UInteractiveTool* GetParentTool() const;
 
 	/**
 	 * Add a PropertySet object for this Mechanic to parent Tool
 	 * @param PropertySet Property Set object to add
 	 */
-	virtual void AddToolPropertySource(UInteractiveToolPropertySet* PropertySet);
+	INTERACTIVETOOLSFRAMEWORK_API virtual void AddToolPropertySource(UInteractiveToolPropertySet* PropertySet);
 
 	/** Enable/Disable a PropertySet object for the parent tool. */
-	virtual bool SetToolPropertySourceEnabled(UInteractiveToolPropertySet* PropertySet, bool bEnabled);
+	INTERACTIVETOOLSFRAMEWORK_API virtual bool SetToolPropertySourceEnabled(UInteractiveToolPropertySet* PropertySet, bool bEnabled);
 };
 
diff --git a/Engine/Source/Runtime/InteractiveToolsFramework/Public/InteractiveCommand.h b/Engine/Source/Runtime/InteractiveToolsFramework/Public/InteractiveCommand.h
--- a/Engine/Source/Runtime/InteractiveToolsFramework/Public/InteractiveCommand.h
+++ b/Engine/Source/Runtime/InteractiveToolsFramework/Public/InteractiveCommand.h
@@ -14,8 +14,8 @@ class IToolsContextTransactionsAPI;
  * The base UInteractiveCommandArguments provides support for transactions via 
  * a IToolsContextTransactionsAPI
  */
-UCLASS()
-class INTERACTIVETOOLSFRAMEWORK_API UInteractiveCommandArguments : public UObject
+UCLASS(MinimalAPI)
+class UInteractiveCommandArguments : public UObject
 {
 	GENERATED_BODY()
 
@@ -34,8 +34,8 @@ protected:
  * UInteractiveCommandResult subclasses are returned from UInteractiveCommands, to allow
  * commands to return custom information.
  */
-UCLASS()
-class INTERACTIVETOOLSFRAMEWORK_API UInteractiveCommandResult : public UObject
+UCLASS(MinimalAPI)
+class UInteractiveCommandResult : public UObject
 {
 	GENERATED_BODY()
 
@@ -48,8 +48,8 @@ class INTERACTIVETOOLSFRAMEWORK_API UInteractiveCommandResult : public UObject
  * This differs from an Interactive Tool in that there is no ongoing user interaction once the
  * command has been initiated. 
  */
-UCLASS(Abstract)
-class INTERACTIVETOOLSFRAMEWORK_API UInteractiveCommand : public UObject
+UCLASS(Abstract, MinimalAPI)
+class UInteractiveCommand : public UObject
 {
 	GENERATED_BODY()
 
@@ -78,4 +78,4 @@ public:
 	virtual void ExecuteCommand(UInteractiveCommandArguments* Arguments, UInteractiveCommandResult** Result = nullptr)
 	{
 	}
-};
\ No newline at end of file
+};
diff --git a/Engine/Source/Runtime/InteractiveToolsFramework/Public/InteractiveGizmo.h b/Engine/Source/Runtime/InteractiveToolsFramework/Public/InteractiveGizmo.h
--- a/Engine/Source/Runtime/InteractiveToolsFramework/Public/InteractiveGizmo.h
+++ b/Engine/Source/Runtime/InteractiveToolsFramework/Public/InteractiveGizmo.h
@@ -24,49 +24,49 @@ class UInteractiveGizmoManager;
  * @todo callback/delegate for if/when .InputBehaviors changes
  * @todo callback/delegate for when Gizmo properties change
  */
-UCLASS(Transient)
-class INTERACTIVETOOLSFRAMEWORK_API UInteractiveGizmo : public UObject, public IInputBehaviorSource
+UCLASS(Transient, MinimalAPI)
+class UInteractiveGizmo : public UObject, public IInputBehaviorSource
 {
 	GENERATED_BODY()
 
 public:
-	UInteractiveGizmo();
+	INTERACTIVETOOLSFRAMEWORK_API UInteractiveGizmo();
 
 	/**
 	 * Called by GizmoManager to initialize the Gizmo *after* GizmoBuilder::BuildGizmo() has been called
 	 */
-	virtual void Setup();
+	INTERACTIVETOOLSFRAMEWORK_API virtual void Setup();
 
 	/**
 	 * Called by GizmoManager to shut down the Gizmo
 	 */
-	virtual void Shutdown();
+	INTERACTIVETOOLSFRAMEWORK_API virtual void Shutdown();
 
 	/**
 	 * Allow the Gizmo to do any custom drawing (ie via PDI/RHI)
 	 * @param RenderAPI Abstraction that provides access to Rendering in the current ToolsContext
 	 */
-	virtual void Render(IToolsContextRenderAPI* RenderAPI);
+	INTERACTIVETOOLSFRAMEWORK_API virtual void Render(IToolsContextRenderAPI* RenderAPI);
 
 	/**
 	 * Allow the Gizmo to do any custom screen space drawing
 	 * @param Canvas the FCanvas to use to do the drawing
 	 * @param RenderAPI Abstraction that provides access to Rendering in the current ToolsContext
 	 */
-	virtual void DrawHUD( FCanvas* Canvas, IToolsContextRenderAPI* RenderAPI );
+	INTERACTIVETOOLSFRAMEWORK_API virtual void DrawHUD( FCanvas* Canvas, IToolsContextRenderAPI* RenderAPI );
 
 	/**
 	 * Allow the Gizmo to do any necessary processing on Tick
 	 * @param DeltaTime the time delta since last tick
 	 */
-	virtual void Tick(float DeltaTime);
+	INTERACTIVETOOLSFRAMEWORK_API virtual void Tick(float DeltaTime);
 
 
 
 	/**
 	 * @return GizmoManager that owns this Gizmo
 	 */
-	virtual UInteractiveGizmoManager* GetGizmoManager() const;
+	INTERACTIVETOOLSFRAMEWORK_API virtual UInteractiveGizmoManager* GetGizmoManager() const;
 
 
 
@@ -78,12 +78,12 @@ public:
 	 * Add an input behavior for this Gizmo
 	 * @param Behavior behavior to add
 	 */
-	virtual void AddInputBehavior(UInputBehavior* Behavior);
+	INTERACTIVETOOLSFRAMEWORK_API virtual void AddInputBehavior(UInputBehavior* Behavior);
 
 	/**
 	 * @return Current input behavior set.
 	 */
-	virtual const UInputBehaviorSet* GetInputBehaviors() const;
+	INTERACTIVETOOLSFRAMEWORK_API virtual const UInputBehaviorSet* GetInputBehaviors() const;
 
 
 
diff --git a/Engine/Source/Runtime/InteractiveToolsFramework/Public/InteractiveGizmoBuilder.h b/Engine/Source/Runtime/InteractiveToolsFramework/Public/InteractiveGizmoBuilder.h
--- a/Engine/Source/Runtime/InteractiveToolsFramework/Public/InteractiveGizmoBuilder.h
+++ b/Engine/Source/Runtime/InteractiveToolsFramework/Public/InteractiveGizmoBuilder.h
@@ -14,8 +14,8 @@
  * These are registered with the InteractiveGizmoManager, which calls BuildGizmo().
  * This is an abstract base class, you must subclass it in order to create your particular Gizmo instance
  */
-UCLASS(Transient, Abstract)
-class INTERACTIVETOOLSFRAMEWORK_API UInteractiveGizmoBuilder : public UObject
+UCLASS(Transient, Abstract, MinimalAPI)
+class UInteractiveGizmoBuilder : public UObject
 {
 	GENERATED_BODY()
 
@@ -31,4 +31,4 @@ public:
 		unimplemented();
 		return nullptr;
 	}
-};
\ No newline at end of file
+};
diff --git a/Engine/Source/Runtime/InteractiveToolsFramework/Public/InteractiveGizmoManager.h b/Engine/Source/Runtime/InteractiveToolsFramework/Public/InteractiveGizmoManager.h
--- a/Engine/Source/Runtime/InteractiveToolsFramework/Public/InteractiveGizmoManager.h
+++ b/Engine/Source/Runtime/InteractiveToolsFramework/Public/InteractiveGizmoManager.h
@@ -35,21 +35,21 @@ struct FActiveGizmo
  * Gizmos can then be activated via the string identifier.
  * 
  */
-UCLASS(Transient)
-class INTERACTIVETOOLSFRAMEWORK_API UInteractiveGizmoManager : public UObject, public IToolContextTransactionProvider
+UCLASS(Transient, MinimalAPI)
+class UInteractiveGizmoManager : public UObject, public IToolContextTransactionProvider
 {
 	GENERATED_BODY()
 
 protected:
 	friend class UInteractiveToolsContext;		// to call Initialize/Shutdown
 
-	UInteractiveGizmoManager();
+	INTERACTIVETOOLSFRAMEWORK_API UInteractiveGizmoManager();
 
 	/** Initialize the GizmoManager with the necessary Context-level state. UInteractiveToolsContext calls this, you should not. */
-	virtual void Initialize(IToolsContextQueriesAPI* QueriesAPI, IToolsContextTransactionsAPI* TransactionsAPI, UInputRouter* InputRouter);
+	INTERACTIVETOOLSFRAMEWORK_API virtual void Initialize(IToolsContextQueriesAPI* QueriesAPI, IToolsContextTransactionsAPI* TransactionsAPI, UInputRouter* InputRouter);
 
 	/** Shutdown the GizmoManager. Called by UInteractiveToolsContext. */
-	virtual void Shutdown();
+	INTERACTIVETOOLSFRAMEWORK_API virtual void Shutdown();
 
 public:
 
@@ -62,14 +62,14 @@ public:
 	 * @param BuilderIdentifier string used to identify this Builder
 	 * @param Builder new GizmoBuilder instance
 	 */
-	virtual void RegisterGizmoType(const FString& BuilderIdentifier, UInteractiveGizmoBuilder* Builder);
+	INTERACTIVETOOLSFRAMEWORK_API virtual void RegisterGizmoType(const FString& BuilderIdentifier, UInteractiveGizmoBuilder* Builder);
 
 	/**
 	 * Remove a GizmoBuilder from the set of known GizmoBuilders
 	 * @param BuilderIdentifier identification string that was passed to RegisterGizmoType()
 	 * @return true if Builder was found and deregistered
 	 */
-	virtual bool DeregisterGizmoType(const FString& BuilderIdentifier);
+	INTERACTIVETOOLSFRAMEWORK_API virtual bool DeregisterGizmoType(const FString& BuilderIdentifier);
 
 
 
@@ -80,7 +80,7 @@ public:
 	 * @param Owner void pointer to whatever "owns" this Gizmo. Allows Gizmo to later be deleted using DestroyAllGizmosByOwner()
 	 * @return new Gizmo instance that has been created and initialized
 	 */	
-	virtual UInteractiveGizmo* CreateGizmo(const FString& BuilderIdentifier, const FString& InstanceIdentifier = FString(), void* Owner = nullptr);
+	INTERACTIVETOOLSFRAMEWORK_API virtual UInteractiveGizmo* CreateGizmo(const FString& BuilderIdentifier, const FString& InstanceIdentifier = FString(), void* Owner = nullptr);
 
 
 	/**
@@ -103,19 +103,19 @@ public:
 	 * @param Gizmo the Gizmo to shutdown and remove
 	 * @return true if the Gizmo was found and removed
 	 */
-	virtual bool DestroyGizmo(UInteractiveGizmo* Gizmo);
+	INTERACTIVETOOLSFRAMEWORK_API virtual bool DestroyGizmo(UInteractiveGizmo* Gizmo);
 
 	/**
 	 * Destroy all Gizmos that were created by the identified GizmoBuilder
 	 * @param BuilderIdentifier the Builder string registered with RegisterGizmoType
 	 */
-	virtual void DestroyAllGizmosOfType(const FString& BuilderIdentifier);
+	INTERACTIVETOOLSFRAMEWORK_API virtual void DestroyAllGizmosOfType(const FString& BuilderIdentifier);
 
 	/**
 	 * Destroy all Gizmos that are owned by the given pointer
 	 * @param Owner pointer that was passed to CreateGizmo
 	 */
-	virtual void DestroyAllGizmosByOwner(void* Owner);
+	INTERACTIVETOOLSFRAMEWORK_API virtual void DestroyAllGizmosByOwner(void* Owner);
 
 
 	/**
@@ -123,14 +123,14 @@ public:
 	 * @param BuilderIdentifier the Builder string registered with RegisterGizmoType
 	 * @return list of found Gizmos
 	 */
-	virtual TArray<UInteractiveGizmo*> FindAllGizmosOfType(const FString& BuilderIdentifier);
+	INTERACTIVETOOLSFRAMEWORK_API virtual TArray<UInteractiveGizmo*> FindAllGizmosOfType(const FString& BuilderIdentifier);
 
 	/**
 	 * Find the Gizmo that was created with the given instance identifier
 	 * @param Identifier the InstanceIdentifier that was passed to CreateGizmo()
 	 * @return the found Gizmo, or null
 	 */
-	virtual UInteractiveGizmo* FindGizmoByInstanceIdentifier(const FString& Identifier);
+	INTERACTIVETOOLSFRAMEWORK_API virtual UInteractiveGizmo* FindGizmoByInstanceIdentifier(const FString& Identifier);
 
 
 
@@ -140,19 +140,19 @@ public:
 	//
 	
 	/** Post a message via the Transactions API */
-	virtual void DisplayMessage(const FText& Message, EToolMessageLevel Level);
+	INTERACTIVETOOLSFRAMEWORK_API virtual void DisplayMessage(const FText& Message, EToolMessageLevel Level);
 
 	/** Request an Invalidation via the Transactions API (ie to cause a repaint, etc) */
-	virtual void PostInvalidation();
+	INTERACTIVETOOLSFRAMEWORK_API virtual void PostInvalidation();
 
 	/**
 	 * Request that the Context open a Transaction, whatever that means to the current Context
 	 * @param Description text description of this transaction (this is the string that appears on undo/redo in the UE Editor)
 	 */
-	virtual void BeginUndoTransaction(const FText& Description);
+	INTERACTIVETOOLSFRAMEWORK_API virtual void BeginUndoTransaction(const FText& Description);
 
 	/** Request that the Context close and commit the open Transaction */
-	virtual void EndUndoTransaction();
+	INTERACTIVETOOLSFRAMEWORK_API virtual void EndUndoTransaction();
 
 	/**
 	 * Forward an FChange object to the Context
@@ -160,7 +160,7 @@ public:
 	 * @param Change the change object that the Context should insert into the transaction history
 	 * @param Description text description of this change (this is the string that appears on undo/redo in the UE Editor)
 	 */
-	virtual void EmitObjectChange(UObject* TargetObject, TUniquePtr<FToolCommandChange> Change, const FText& Description );
+	INTERACTIVETOOLSFRAMEWORK_API virtual void EmitObjectChange(UObject* TargetObject, TUniquePtr<FToolCommandChange> Change, const FText& Description );
 
 
 	//
@@ -168,13 +168,13 @@ public:
 	//
 
 	/** Tick any active Gizmos. Called by UInteractiveToolsContext */
-	virtual void Tick(float DeltaTime);
+	INTERACTIVETOOLSFRAMEWORK_API virtual void Tick(float DeltaTime);
 
 	/** Render any active Gizmos. Called by UInteractiveToolsContext. */
-	virtual void Render(IToolsContextRenderAPI* RenderAPI);
+	INTERACTIVETOOLSFRAMEWORK_API virtual void Render(IToolsContextRenderAPI* RenderAPI);
 
 	/** Let active Gizmos do screen space drawing.  Called by UInteractiveToolsContext */
-	virtual void DrawHUD( FCanvas* Canvas, IToolsContextRenderAPI* RenderAPI );
+	INTERACTIVETOOLSFRAMEWORK_API virtual void DrawHUD( FCanvas* Canvas, IToolsContextRenderAPI* RenderAPI );
 
 	//
 	// access to APIs, etc
@@ -186,7 +186,7 @@ public:
 	/**
 	 * @return the context object store from the owning tools context.
 	 */
-	UContextObjectStore* GetContextObjectStore() const;
+	INTERACTIVETOOLSFRAMEWORK_API UContextObjectStore* GetContextObjectStore() const;
 
 
 
@@ -198,7 +198,7 @@ public:
 	/**
 	 * Register default gizmo types
 	 */
-	virtual void RegisterDefaultGizmos();
+	INTERACTIVETOOLSFRAMEWORK_API virtual void RegisterDefaultGizmos();
 
 	/**
 	 * Activate a new instance of the default 3-axis transformation Gizmo. RegisterDefaultGizmos() must have been called first.
@@ -206,7 +206,7 @@ public:
 	 * @param InstanceIdentifier optional client-defined *unique* string that can be used to locate this instance
 	 * @return new Gizmo instance that has been created and initialized
 	 */
-	virtual UCombinedTransformGizmo* Create3AxisTransformGizmo(void* Owner = nullptr, const FString& InstanceIdentifier = FString());
+	INTERACTIVETOOLSFRAMEWORK_API virtual UCombinedTransformGizmo* Create3AxisTransformGizmo(void* Owner = nullptr, const FString& InstanceIdentifier = FString());
 
 	/**
 	 * Activate a new customized instance of the default 3-axis transformation Gizmo, with only certain elements included. RegisterDefaultGizmos() must have been called first.
@@ -215,19 +215,19 @@ public:
 	 * @param InstanceIdentifier optional client-defined *unique* string that can be used to locate this instance
 	 * @return new Gizmo instance that has been created and initialized
 	 */
-	virtual UCombinedTransformGizmo* CreateCustomTransformGizmo(ETransformGizmoSubElements Elements, void* Owner = nullptr, const FString& InstanceIdentifier = FString());
+	INTERACTIVETOOLSFRAMEWORK_API virtual UCombinedTransformGizmo* CreateCustomTransformGizmo(ETransformGizmoSubElements Elements, void* Owner = nullptr, const FString& InstanceIdentifier = FString());
 
-	virtual UCombinedTransformGizmo* CreateCustomRepositionableTransformGizmo(ETransformGizmoSubElements Elements, void* Owner = nullptr, const FString& InstanceIdentifier = FString());
+	INTERACTIVETOOLSFRAMEWORK_API virtual UCombinedTransformGizmo* CreateCustomRepositionableTransformGizmo(ETransformGizmoSubElements Elements, void* Owner = nullptr, const FString& InstanceIdentifier = FString());
 
 public:
 	// builder identifiers for default gizmo types. Perhaps should have an API for this...
-	static FString DefaultAxisPositionBuilderIdentifier;
-	static FString DefaultPlanePositionBuilderIdentifier;
-	static FString DefaultAxisAngleBuilderIdentifier;
-	static FString DefaultThreeAxisTransformBuilderIdentifier;
-	static const FString CustomThreeAxisTransformBuilderIdentifier;
-	static const FString CustomRepositionableThreeAxisTransformBuilderIdentifier;
-	static FString DefaultScalableSphereBuilderIdentifier;
+	static INTERACTIVETOOLSFRAMEWORK_API FString DefaultAxisPositionBuilderIdentifier;
+	static INTERACTIVETOOLSFRAMEWORK_API FString DefaultPlanePositionBuilderIdentifier;
+	static INTERACTIVETOOLSFRAMEWORK_API FString DefaultAxisAngleBuilderIdentifier;
+	static INTERACTIVETOOLSFRAMEWORK_API FString DefaultThreeAxisTransformBuilderIdentifier;
+	static INTERACTIVETOOLSFRAMEWORK_API const FString CustomThreeAxisTransformBuilderIdentifier;
+	static INTERACTIVETOOLSFRAMEWORK_API const FString CustomRepositionableThreeAxisTransformBuilderIdentifier;
+	static INTERACTIVETOOLSFRAMEWORK_API FString DefaultScalableSphereBuilderIdentifier;
 
 protected:
 	/** set of Currently-active Gizmos */
diff --git a/Engine/Source/Runtime/InteractiveToolsFramework/Public/InteractiveTool.h b/Engine/Source/Runtime/InteractiveToolsFramework/Public/InteractiveTool.h
--- a/Engine/Source/Runtime/InteractiveToolsFramework/Public/InteractiveTool.h
+++ b/Engine/Source/Runtime/InteractiveToolsFramework/Public/InteractiveTool.h
@@ -53,14 +53,14 @@ enum class EToolShutdownType : uint8
 /**
  * FInteractiveToolInfo provides information about a tool (name, tooltip, etc)
  */
-struct INTERACTIVETOOLSFRAMEWORK_API FInteractiveToolInfo
+struct FInteractiveToolInfo
 {
 	/** Name of Tool. May be FText::Empty(), but will default to Tool->GetClass()->GetDisplayNameText() in InteractiveTool constructor */
 	FText ToolDisplayName = FText::GetEmpty();
 };
 
 
-class INTERACTIVETOOLSFRAMEWORK_API FWatchablePropertySet
+class FWatchablePropertySet
 {
 public:
 	//
@@ -238,8 +238,8 @@ DECLARE_MULTICAST_DELEGATE_TwoParams(FInteractiveToolPropertySetModifiedSignatur
  * when the Tool Properties are changed *in the Editor*, the parent Tool will be automatically notified.
  * You can override UInteractiveTool::OnPropertyModified() to act on these notifications
  */
-UCLASS(Transient)
-class INTERACTIVETOOLSFRAMEWORK_API UInteractiveToolPropertySet : public UObject, public FWatchablePropertySet
+UCLASS(Transient, MinimalAPI)
+class UInteractiveToolPropertySet : public UObject, public FWatchablePropertySet
 {
 	GENERATED_BODY()
 
@@ -282,17 +282,17 @@ public:
 	 * Save the values of this PropertySet with the given CacheIdentifier.
 	 * The Tool parameter is currently ignored.
 	 */
-	virtual void SaveProperties(UInteractiveTool* SaveFromTool, const FString& CacheIdentifier = TEXT(""));
+	INTERACTIVETOOLSFRAMEWORK_API virtual void SaveProperties(UInteractiveTool* SaveFromTool, const FString& CacheIdentifier = TEXT(""));
 
 	/**
 	 * Restore the values of the Property Set with the given CacheIdentifier.
 	 * The Tool parameter is currently ignored.
 	 */
-	virtual void RestoreProperties(UInteractiveTool* RestoreToTool, const FString& CacheIdentifier = TEXT(""));
+	INTERACTIVETOOLSFRAMEWORK_API virtual void RestoreProperties(UInteractiveTool* RestoreToTool, const FString& CacheIdentifier = TEXT(""));
 
 protected:
 	// Utility func used to implement the default Save/RestoreProperties funcs
-	virtual void SaveRestoreProperties(UInteractiveTool* RestoreToTool, const FString& CacheIdentifier, bool bSaving);
+	INTERACTIVETOOLSFRAMEWORK_API virtual void SaveRestoreProperties(UInteractiveTool* RestoreToTool, const FString& CacheIdentifier, bool bSaving);
 
 	/**
 	 * GetDynamicPropertyCache return class-internal objects that subclasses can use to save/restore properties.
@@ -300,7 +300,7 @@ protected:
 	 * @param bWasCreatedOut true is returned here if this is the first time the object was seen
 	 * @return instance of the current subclass that can be used to save/restore values
 	 */
-	TObjectPtr<UInteractiveToolPropertySet> GetDynamicPropertyCache(const FString& CacheIdentifier, bool& bWasCreatedOut);
+	INTERACTIVETOOLSFRAMEWORK_API TObjectPtr<UInteractiveToolPropertySet> GetDynamicPropertyCache(const FString& CacheIdentifier, bool& bWasCreatedOut);
 
 public:
 #if WITH_EDITOR
@@ -309,7 +309,7 @@ public:
 	  * @warning Please consider listening to OnModified instead of overriding this function
 	  * @warning this function is currently only called in Editor (not at runtime)
 	  */
-	virtual void PostEditChangeProperty(FPropertyChangedEvent& PropertyChangedEvent);
+	INTERACTIVETOOLSFRAMEWORK_API virtual void PostEditChangeProperty(FPropertyChangedEvent& PropertyChangedEvent);
 #endif
 
 protected:
@@ -345,64 +345,64 @@ protected:
  * @todo callback/delegate for if/when .InputBehaviors changes
  * @todo callback/delegate for when tool properties change
  */
-UCLASS(Transient)
-class INTERACTIVETOOLSFRAMEWORK_API UInteractiveTool : public UObject, public IInputBehaviorSource
+UCLASS(Transient, MinimalAPI)
+class UInteractiveTool : public UObject, public IInputBehaviorSource
 {
 	GENERATED_BODY()
 
 public:
-	UInteractiveTool();
+	INTERACTIVETOOLSFRAMEWORK_API UInteractiveTool();
 
 	/**
 	 * Called by ToolManager to initialize the Tool *after* ToolBuilder::BuildTool() has been called
 	 */
-	virtual void Setup();
+	INTERACTIVETOOLSFRAMEWORK_API virtual void Setup();
 
 	/**
 	 * Called by ToolManager to shut down the Tool
 	 * @param ShutdownType indicates how the tool should shutdown (ie Accept or Cancel current preview, etc)
 	 */
-	virtual void Shutdown(EToolShutdownType ShutdownType);
+	INTERACTIVETOOLSFRAMEWORK_API virtual void Shutdown(EToolShutdownType ShutdownType);
 
 	/**
 	 * Allow the Tool to do any custom drawing (ie via PDI/RHI)
 	 * @param RenderAPI Abstraction that provides access to Rendering in the current ToolsContext
 	 */
-	virtual void Render(IToolsContextRenderAPI* RenderAPI);
+	INTERACTIVETOOLSFRAMEWORK_API virtual void Render(IToolsContextRenderAPI* RenderAPI);
 
 	/**
 	 * Allow the Tool to do any custom screen space drawing
 	 * @param Canvas the FCanvas to use to do the drawing
 	 * @param RenderAPI Abstraction that provides access to Rendering in the current ToolsContext
 	 */
-	virtual void DrawHUD( FCanvas* Canvas, IToolsContextRenderAPI* RenderAPI );
+	INTERACTIVETOOLSFRAMEWORK_API virtual void DrawHUD( FCanvas* Canvas, IToolsContextRenderAPI* RenderAPI );
 
 	/**
 	 * Non overrideable func which does processing and calls the tool's OnTick
 	 * @param DeltaTime the time delta since last tick
 	 */
-	virtual void Tick(float DeltaTime) final;
+	INTERACTIVETOOLSFRAMEWORK_API virtual void Tick(float DeltaTime) final;
 
 	/**
 	 * @return ToolManager that owns this Tool
 	 */
-	virtual UInteractiveToolManager* GetToolManager() const;
+	INTERACTIVETOOLSFRAMEWORK_API virtual UInteractiveToolManager* GetToolManager() const;
 
 
 	/**
 	 * @return true if this Tool support being Cancelled, ie calling Shutdown(EToolShutdownType::Cancel)
 	 */
-	virtual bool HasCancel() const;
+	INTERACTIVETOOLSFRAMEWORK_API virtual bool HasCancel() const;
 
 	/**
 	 * @return true if this Tool support being Accepted, ie calling Shutdown(EToolShutdownType::Accept)
 	 */
-	virtual bool HasAccept() const;
+	INTERACTIVETOOLSFRAMEWORK_API virtual bool HasAccept() const;
 
 	/**
 	 * @return true if this Tool is currently in a state where it can be Accepted. This may be false if for example there was an error in the Tool.
 	 */
-	virtual bool CanAccept() const;
+	INTERACTIVETOOLSFRAMEWORK_API virtual bool CanAccept() const;
 
 
 
@@ -416,7 +416,7 @@ public:
 	 * @param Behavior behavior to add
 	 * @param Source Optional pointer that could be used to identify the behavior for removal later.
 	 */
-	virtual void AddInputBehavior(UInputBehavior* Behavior, void* Source = nullptr);
+	INTERACTIVETOOLSFRAMEWORK_API virtual void AddInputBehavior(UInputBehavior* Behavior, void* Source = nullptr);
 
 	/**
 	 * Remove all input behaviors that had the given source pointer set during their addition.
@@ -424,12 +424,12 @@ public:
 	 * to the input router after that call.
 	 * @param Source Identifying pointer
 	 */
-	virtual void RemoveInputBehaviorsBySource(void* Source);
+	INTERACTIVETOOLSFRAMEWORK_API virtual void RemoveInputBehaviorsBySource(void* Source);
 
 	/**
 	 * @return Current input behavior set.
 	 */
-	virtual const UInputBehaviorSet* GetInputBehaviors() const;
+	INTERACTIVETOOLSFRAMEWORK_API virtual const UInputBehaviorSet* GetInputBehaviors() const;
 
 
 	//
@@ -439,7 +439,7 @@ public:
 	/**
 	 * @return list of property UObjects for this tool (ie to add to a DetailsViewPanel, for example)
 	 */
-	virtual TArray<UObject*> GetToolProperties(bool bEnabledOnly = true) const;
+	INTERACTIVETOOLSFRAMEWORK_API virtual TArray<UObject*> GetToolProperties(bool bEnabledOnly = true) const;
 
 	/**
 	 * OnPropertySetsModified is broadcast whenever the contents of the ToolPropertyObjects array is modified
@@ -483,20 +483,20 @@ protected:
 	 * Add a Property object for this Tool
 	 * @param Property object to add
 	 */
-	virtual void AddToolPropertySource(UObject* PropertyObject);
+	INTERACTIVETOOLSFRAMEWORK_API virtual void AddToolPropertySource(UObject* PropertyObject);
 
 	/**
 	 * Add a PropertySet object for this Tool
 	 * @param PropertySet Property Set object to add
 	 */
-	virtual void AddToolPropertySource(UInteractiveToolPropertySet* PropertySet);
+	INTERACTIVETOOLSFRAMEWORK_API virtual void AddToolPropertySource(UInteractiveToolPropertySet* PropertySet);
 
 	/**
 	 * Remove a PropertySet object from this Tool. If found, will broadcast OnPropertySetsModified
 	 * @param PropertySet property set to remove.
 	 * @return true if PropertySet is found and removed
 	 */
-	virtual bool RemoveToolPropertySource(UInteractiveToolPropertySet* PropertySet);
+	INTERACTIVETOOLSFRAMEWORK_API virtual bool RemoveToolPropertySource(UInteractiveToolPropertySet* PropertySet);
 
 	/**
 	 * Replace a PropertySet object on this Tool with another property set. If replaced, will broadcast OnPropertySetsModified
@@ -505,7 +505,7 @@ protected:
 	 * @param bSetToEnabled if true, ReplaceWith property set is explicitly enabled (otherwise enable/disable state is unmodified)
 	 * @return true if CurPropertySet is found and replaced
 	 */
-	virtual bool ReplaceToolPropertySource(UInteractiveToolPropertySet* CurPropertySet, UInteractiveToolPropertySet* ReplaceWith, bool bSetToEnabled = true);
+	INTERACTIVETOOLSFRAMEWORK_API virtual bool ReplaceToolPropertySource(UInteractiveToolPropertySet* CurPropertySet, UInteractiveToolPropertySet* ReplaceWith, bool bSetToEnabled = true);
 
 	/**
 	 * Enable/Disable a PropertySet object for this Tool. If found and state was modified, will broadcast OnPropertySetsModified
@@ -513,7 +513,7 @@ protected:
 	 * @param bEnabled whether to enable or disable
 	 * @return true if PropertySet was found
 	 */
-	virtual bool SetToolPropertySourceEnabled(UInteractiveToolPropertySet* PropertySet, bool bEnabled);
+	INTERACTIVETOOLSFRAMEWORK_API virtual bool SetToolPropertySourceEnabled(UInteractiveToolPropertySet* PropertySet, bool bEnabled);
 
 	/**
 	 * Call after changing a propertyset internally in the tool to allow external views of the property
@@ -521,7 +521,7 @@ protected:
 	 * pass along notifications, so don't call this if the property is changed externally (i.e., this
 	 * should not usually be called from OnPropertyModified unless the tool adds changes of its own).
 	 */
-	virtual void NotifyOfPropertyChangeByTool(UInteractiveToolPropertySet* PropertySet) const;
+	INTERACTIVETOOLSFRAMEWORK_API virtual void NotifyOfPropertyChangeByTool(UInteractiveToolPropertySet* PropertySet) const;
 
 	enum EAcceptWarning
 	{
@@ -536,7 +536,7 @@ protected:
 	 *
 	 * @param Warning Reason that tool cannot be accepted (or EAcceptWarning::NoWarning if no warning need be displayed)
 	 */
-	virtual void UpdateAcceptWarnings(EAcceptWarning Warning);
+	INTERACTIVETOOLSFRAMEWORK_API virtual void UpdateAcceptWarnings(EAcceptWarning Warning);
 
 private:
 	// Tracks whether the UpdateAcceptWarnings function showed a warning the last time it was called.
@@ -566,13 +566,13 @@ public:
 	 * Get the internal Action Set for this Tool. The action set is created and registered on-demand.
 	 * @return pointer to initialized Action set
 	 */
-	virtual FInteractiveToolActionSet* GetActionSet();
+	INTERACTIVETOOLSFRAMEWORK_API virtual FInteractiveToolActionSet* GetActionSet();
 
 	/**
 	 * Request that the Action identified by ActionID be executed.
 	 * Default implementation forwards these requests to internal ToolActionSet.
 	 */
-	virtual void ExecuteAction(int32 ActionID);
+	INTERACTIVETOOLSFRAMEWORK_API virtual void ExecuteAction(int32 ActionID);
 
 
 protected:
@@ -581,7 +581,7 @@ protected:
 	 * Note that for the actions to be triggered, you will also need to add corresponding registration per tool
 	 *  -- see Engine\Plugins\Editor\ModelingToolsEditorMode\Source\ModelingToolsEditorMode\Public\ModelingToolsActions.h for examples
 	 */
-	virtual void RegisterActions(FInteractiveToolActionSet& ActionSet);
+	INTERACTIVETOOLSFRAMEWORK_API virtual void RegisterActions(FInteractiveToolActionSet& ActionSet);
 
 
 private:
diff --git a/Engine/Source/Runtime/InteractiveToolsFramework/Public/InteractiveToolActionSet.h b/Engine/Source/Runtime/InteractiveToolsFramework/Public/InteractiveToolActionSet.h
--- a/Engine/Source/Runtime/InteractiveToolsFramework/Public/InteractiveToolActionSet.h
+++ b/Engine/Source/Runtime/InteractiveToolsFramework/Public/InteractiveToolActionSet.h
@@ -36,7 +36,7 @@ enum class EStandardToolActions
  * FInteractiveToolAction is returned by a UInteractiveTool to represent
  * an "Action" the Tool can execute. 
  */
-struct INTERACTIVETOOLSFRAMEWORK_API FInteractiveToolAction
+struct FInteractiveToolAction
 {
 	/** Which type of UInteractiveTool this Action can be applied to*/
 	const UClass* ClassType;
@@ -87,7 +87,7 @@ struct INTERACTIVETOOLSFRAMEWORK_API FInteractiveToolAction
  * FInteractiveToolActionSet maintains a list of FInteractiveToolAction.
  * Each UInteractiveTool contains an instance of this class.
  */
-class INTERACTIVETOOLSFRAMEWORK_API FInteractiveToolActionSet
+class FInteractiveToolActionSet
 {
 public:
 
@@ -95,7 +95,7 @@ public:
 	 * Register an Action with the ActionSet. This function is intended to be called by
 	 * UInteractiveTool::RegisterActions() implementations
 	 */
-	void RegisterAction(UInteractiveTool* Tool, int32 ActionID,
+	INTERACTIVETOOLSFRAMEWORK_API void RegisterAction(UInteractiveTool* Tool, int32 ActionID,
 		const FString& ActionName, const FText& ShortUIName, const FText& DescriptionText,
 		EModifierKey::Type Modifiers, const FKey& ShortcutKey,
 		TFunction<void()> ActionFunction );
@@ -104,18 +104,18 @@ public:
 	 * Find an existing Action by ID
 	 * @return located Action, or nullptr if not found
 	 */
-	const FInteractiveToolAction* FindActionByID(int32 ActionID) const;
+	INTERACTIVETOOLSFRAMEWORK_API const FInteractiveToolAction* FindActionByID(int32 ActionID) const;
 
 	/**
 	 * Return the internal list of registered Actions by adding to the OutActions array
 	 */
-	void CollectActions(TArray<FInteractiveToolAction>& OutActions) const;
+	INTERACTIVETOOLSFRAMEWORK_API void CollectActions(TArray<FInteractiveToolAction>& OutActions) const;
 
 	/**
 	 * Execute the action identified by ActionID
 	 */
-	void ExecuteAction(int32 ActionID) const;
+	INTERACTIVETOOLSFRAMEWORK_API void ExecuteAction(int32 ActionID) const;
 
 protected:
 	TArray<FInteractiveToolAction> Actions;
-};
\ No newline at end of file
+};
diff --git a/Engine/Source/Runtime/InteractiveToolsFramework/Public/InteractiveToolBuilder.h b/Engine/Source/Runtime/InteractiveToolsFramework/Public/InteractiveToolBuilder.h
--- a/Engine/Source/Runtime/InteractiveToolsFramework/Public/InteractiveToolBuilder.h
+++ b/Engine/Source/Runtime/InteractiveToolsFramework/Public/InteractiveToolBuilder.h
@@ -23,8 +23,8 @@ struct FToolBuilderState;
  * In addition CanBuildTool() will be queried to (for example) enable/disable UI buttons, etc.
  * This is an abstract base class, you must subclass it in order to create your particular Tool instance
  */
-UCLASS(Transient, Abstract)
-class INTERACTIVETOOLSFRAMEWORK_API UInteractiveToolBuilder : public UObject
+UCLASS(Transient, Abstract, MinimalAPI)
+class UInteractiveToolBuilder : public UObject
 {
 	GENERATED_BODY()
 
@@ -82,8 +82,8 @@ public:
  * interface(s) for defining the tool target requirements of the tool it builds.
  * This is an abstract base class, you must subclass it in order to create your particular Tool instance.
  */
-UCLASS(Transient, Abstract)
-class INTERACTIVETOOLSFRAMEWORK_API UInteractiveToolWithToolTargetsBuilder : public UInteractiveToolBuilder
+UCLASS(Transient, Abstract, MinimalAPI)
+class UInteractiveToolWithToolTargetsBuilder : public UInteractiveToolBuilder
 {
 	GENERATED_BODY()
 
diff --git a/Engine/Source/Runtime/InteractiveToolsFramework/Public/InteractiveToolChange.h b/Engine/Source/Runtime/InteractiveToolsFramework/Public/InteractiveToolChange.h
--- a/Engine/Source/Runtime/InteractiveToolsFramework/Public/InteractiveToolChange.h
+++ b/Engine/Source/Runtime/InteractiveToolsFramework/Public/InteractiveToolChange.h
@@ -14,7 +14,7 @@ class UInteractiveGizmoManager;
 /**
  * FToolCommandChange is a base class for command changes used by the Tools Framework.
  */
-class INTERACTIVETOOLSFRAMEWORK_API FToolCommandChange : public FCommandChange
+class FToolCommandChange : public FCommandChange
 {
 public:
 
@@ -25,8 +25,8 @@ public:
 
 
 
-UINTERFACE()
-class INTERACTIVETOOLSFRAMEWORK_API UToolContextTransactionProvider : public UInterface
+UINTERFACE(MinimalAPI)
+class UToolContextTransactionProvider : public UInterface
 {
 	GENERATED_BODY()
 };
@@ -35,7 +35,7 @@ class INTERACTIVETOOLSFRAMEWORK_API UToolContextTransactionProvider : public UIn
  * uses to interface with the higher-level transaction system. UInteractiveToolManager and UInteractiveGizmoManager
  * both implement this interface.
  */
-class INTERACTIVETOOLSFRAMEWORK_API IToolContextTransactionProvider
+class IToolContextTransactionProvider
 {
 	GENERATED_BODY()
 public:
@@ -64,7 +64,7 @@ public:
  * IToolCommandChangeSource is a generic interface for generating a change object.
  * This is useful to 
  */
-class INTERACTIVETOOLSFRAMEWORK_API IToolCommandChangeSource
+class IToolCommandChangeSource
 {
 public:
 	virtual ~IToolCommandChangeSource() {}
diff --git a/Engine/Source/Runtime/InteractiveToolsFramework/Public/InteractiveToolManager.h b/Engine/Source/Runtime/InteractiveToolsFramework/Public/InteractiveToolManager.h
--- a/Engine/Source/Runtime/InteractiveToolsFramework/Public/InteractiveToolManager.h
+++ b/Engine/Source/Runtime/InteractiveToolsFramework/Public/InteractiveToolManager.h
@@ -58,27 +58,27 @@ enum class EToolChangeTrackingMode
  * on a given side, and then use ActivateTool() to create the new Tool instance.
  *
  */
-UCLASS(Transient)
-class INTERACTIVETOOLSFRAMEWORK_API UInteractiveToolManager : public UObject, public IToolContextTransactionProvider
+UCLASS(Transient, MinimalAPI)
+class UInteractiveToolManager : public UObject, public IToolContextTransactionProvider
 {
 	GENERATED_BODY()
 
 protected:
 	friend class UInteractiveToolsContext;		// to call Initialize/Shutdown
 
-	UInteractiveToolManager();
+	INTERACTIVETOOLSFRAMEWORK_API UInteractiveToolManager();
 
 	/** Initialize the ToolManager with the necessary Context-level state. UInteractiveToolsContext calls this, you should not. */
-	virtual void Initialize(IToolsContextQueriesAPI* QueriesAPI, IToolsContextTransactionsAPI* TransactionsAPI, UInputRouter* InputRouter);
+	INTERACTIVETOOLSFRAMEWORK_API virtual void Initialize(IToolsContextQueriesAPI* QueriesAPI, IToolsContextTransactionsAPI* TransactionsAPI, UInputRouter* InputRouter);
 
 	/** Shutdown the ToolManager. Called by UInteractiveToolsContext. */
-	virtual void Shutdown();
+	INTERACTIVETOOLSFRAMEWORK_API virtual void Shutdown();
 
 	/** Called immediately after a tool is built. Broadcasts OnToolPostBuild. */
-	virtual void DoPostBuild(EToolSide Side, UInteractiveTool* InBuiltTool, UInteractiveToolBuilder* InToolBuilder, const FToolBuilderState& InBuilderState);
+	INTERACTIVETOOLSFRAMEWORK_API virtual void DoPostBuild(EToolSide Side, UInteractiveTool* InBuiltTool, UInteractiveToolBuilder* InToolBuilder, const FToolBuilderState& InBuilderState);
 	
 	/** Called immediately after a tool's Setup is called. Broadcasts OnToolPostSetup. */
-	virtual void DoPostSetup(EToolSide Side, UInteractiveTool* InInteractiveTool, UInteractiveToolBuilder* InToolBuilder, const FToolBuilderState& InBuilderState);
+	INTERACTIVETOOLSFRAMEWORK_API virtual void DoPostSetup(EToolSide Side, UInteractiveTool* InInteractiveTool, UInteractiveToolBuilder* InToolBuilder, const FToolBuilderState& InBuilderState);
 
 public:
 
@@ -98,21 +98,21 @@ public:
 	 * @param Identifier string used to identify this Builder
 	 * @param Builder new ToolBuilder instance
 	 */
-	virtual void RegisterToolType(const FString& Identifier, UInteractiveToolBuilder* Builder);
+	INTERACTIVETOOLSFRAMEWORK_API virtual void RegisterToolType(const FString& Identifier, UInteractiveToolBuilder* Builder);
 
 	/**
 	 * Unregisters a ToolBuilder
 	 * @param Identifier string used to identify this Builder
 	 * @param Builder new ToolBuilder instance
 	 */
-	virtual void UnregisterToolType(const FString& Identifier);
+	INTERACTIVETOOLSFRAMEWORK_API virtual void UnregisterToolType(const FString& Identifier);
 
 	/**
 	 * Set active ToolBuilder for a ToolSide via string identifier
 	 * @param Side which "side" should we set this Builder on
 	 * @param Identifier name of ToolBuilder that was passed to RegisterToolType()
 	 */
-	virtual bool SelectActiveToolType(EToolSide Side, const FString& Identifier);
+	INTERACTIVETOOLSFRAMEWORK_API virtual bool SelectActiveToolType(EToolSide Side, const FString& Identifier);
 
 	/**
 	 * Check if a named Tool type can currently be activated on the given ToolSide
@@ -120,27 +120,27 @@ public:
 	 * @param Identifier string name of the Tool type
 	 * @return true if the Tool type could be activated
 	 */
-	virtual bool CanActivateTool(EToolSide eSide, FString Identifier);
+	INTERACTIVETOOLSFRAMEWORK_API virtual bool CanActivateTool(EToolSide eSide, FString Identifier);
 
 	/**
 	 * Try to activate a new Tool instance on the given Side
 	 * @param Side which "side" you would like to active the tool on
 	 * @return true if a new Tool instance was created and initialized
 	 */	
-	virtual bool ActivateTool(EToolSide Side);
+	INTERACTIVETOOLSFRAMEWORK_API virtual bool ActivateTool(EToolSide Side);
 
 	/**
 	 * Check if there is an active Tool on the given Side
 	 * @param Side which Side to check
 	 * @return true if there is an active Tool on that side
 	 */
-	virtual bool HasActiveTool(EToolSide Side) const;
+	INTERACTIVETOOLSFRAMEWORK_API virtual bool HasActiveTool(EToolSide Side) const;
 
 
 	/**
 	 * @return true if there are any active tools
 	 */
-	virtual bool HasAnyActiveTool() const;
+	INTERACTIVETOOLSFRAMEWORK_API virtual bool HasAnyActiveTool() const;
 
 
 	/**
@@ -148,48 +148,48 @@ public:
 	 * @param Side which Side is being requested
 	 * @return pointer to Tool instance active on that Side, or nullptr if no such Tool exists
 	 */
-	virtual UInteractiveTool* GetActiveTool(EToolSide Side);
+	INTERACTIVETOOLSFRAMEWORK_API virtual UInteractiveTool* GetActiveTool(EToolSide Side);
 
 	/**
 	 * Get pointer to active Tool Builder on a given side
 	 * @param Side which Side is being requested
 	 * @return pointer to Tool Builder instance active on that Side, or nullptr if no such ToolBuilder exists
 	 */
-	virtual UInteractiveToolBuilder* GetActiveToolBuilder(EToolSide Side);
+	INTERACTIVETOOLSFRAMEWORK_API virtual UInteractiveToolBuilder* GetActiveToolBuilder(EToolSide Side);
 
 	/**
 	 * Get name of registered ToolBuilder that created active tool for given side, or empty string if no tool is active
 	 * @param Side which Side is being requested
 	 * @return name of tool, or empty string if no tool is active
 	 */
-	virtual FString GetActiveToolName(EToolSide Side);
+	INTERACTIVETOOLSFRAMEWORK_API virtual FString GetActiveToolName(EToolSide Side);
 
 	/**
 	 * Check if an active Tool on the given Side can be Accepted in its current state
 	 * @param Side which Side to check
 	 * @return true if there is an active Tool and it returns true from HasAccept() and CanAccept()
 	 */
-	virtual bool CanAcceptActiveTool(EToolSide Side);
+	INTERACTIVETOOLSFRAMEWORK_API virtual bool CanAcceptActiveTool(EToolSide Side);
 
 	/**
 	 * Check if an active Tool on the given Side can be Canceled
 	 * @param Side which Side to check
 	 * @return true if there is an active Tool and it returns true from HasCancel()
 	 */
-	virtual bool CanCancelActiveTool(EToolSide Side);
+	INTERACTIVETOOLSFRAMEWORK_API virtual bool CanCancelActiveTool(EToolSide Side);
 
 	/**
 	 * Shut down an active Tool on the given side
 	 * @param Side which "side" you would like to shut down
 	 * @param ShutdownType how should the tool be terminated (eg Accept/Cancel)
 	 */
-	virtual void DeactivateTool(EToolSide Side, EToolShutdownType ShutdownType);
+	INTERACTIVETOOLSFRAMEWORK_API virtual void DeactivateTool(EToolSide Side, EToolShutdownType ShutdownType);
 
 
 	/**
 	 * Configure how tool changes emit change events. See EToolChangeTrackingMode for details.
 	 */
-	virtual void ConfigureChangeTrackingMode(EToolChangeTrackingMode ChangeMode);
+	INTERACTIVETOOLSFRAMEWORK_API virtual void ConfigureChangeTrackingMode(EToolChangeTrackingMode ChangeMode);
 
 
 	//
@@ -197,19 +197,19 @@ public:
 	//
 	
 	/** Post a message via the Transactions API */
-	virtual void DisplayMessage(const FText& Message, EToolMessageLevel Level);
+	INTERACTIVETOOLSFRAMEWORK_API virtual void DisplayMessage(const FText& Message, EToolMessageLevel Level);
 
 	/** Request an Invalidation via the Transactions API (ie to cause a repaint, etc) */
-	virtual void PostInvalidation();
+	INTERACTIVETOOLSFRAMEWORK_API virtual void PostInvalidation();
 
 	/**
 	 * Request that the Context open a Transaction, whatever that means to the current Context
 	 * @param Description text description of this transaction (this is the string that appears on undo/redo in the UE Editor)
 	 */
-	virtual void BeginUndoTransaction(const FText& Description);
+	INTERACTIVETOOLSFRAMEWORK_API virtual void BeginUndoTransaction(const FText& Description);
 
 	/** Request that the Context close and commit the open Transaction */
-	virtual void EndUndoTransaction();
+	INTERACTIVETOOLSFRAMEWORK_API virtual void EndUndoTransaction();
 
 	/**
 	 * Forward an FChange object to the Context
@@ -217,26 +217,26 @@ public:
 	 * @param Change the change object that the Context should insert into the transaction history
 	 * @param Description text description of this change (this is the string that appears on undo/redo in the UE Editor)
 	 */
-	virtual void EmitObjectChange(UObject* TargetObject, TUniquePtr<FToolCommandChange> Change, const FText& Description );
+	INTERACTIVETOOLSFRAMEWORK_API virtual void EmitObjectChange(UObject* TargetObject, TUniquePtr<FToolCommandChange> Change, const FText& Description );
 
 
 	/**
 	 * Forward an FChange object to the Context
 	 */
-	virtual bool RequestSelectionChange(const FSelectedOjectsChangeList& SelectionChange);
+	INTERACTIVETOOLSFRAMEWORK_API virtual bool RequestSelectionChange(const FSelectedOjectsChangeList& SelectionChange);
 
 	//
 	// State control  (@todo: have the Context call these? not safe for anyone to call)
 	//
 
 	/** Tick any active Tools. Called by UInteractiveToolsContext */
-	virtual void Tick(float DeltaTime);
+	INTERACTIVETOOLSFRAMEWORK_API virtual void Tick(float DeltaTime);
 
 	/** Render any active Tools. Called by UInteractiveToolsContext. */
-	virtual void Render(IToolsContextRenderAPI* RenderAPI);
+	INTERACTIVETOOLSFRAMEWORK_API virtual void Render(IToolsContextRenderAPI* RenderAPI);
 
 	/** Let active Tools do their screen space drawing. Called by UInteractiveToolsContext. */
-	virtual void DrawHUD(FCanvas* Canvas, IToolsContextRenderAPI* RenderAPI);
+	INTERACTIVETOOLSFRAMEWORK_API virtual void DrawHUD(FCanvas* Canvas, IToolsContextRenderAPI* RenderAPI);
 
 	/**
 	 * A Tool (or other code) can call this function to request that the Tool be deactivated.
@@ -251,7 +251,7 @@ public:
 	 * @param bShowUnexpectedShutdownMessage if this was an unexpected shutdown, passing true here, along with a non-empty UnexpectedShutdownMessage, will result in OnToolUnexpectedShutdownMessage being called
 	 * @param UnexpectedShutdownMessage message provided by the caller meant to explain why the Tool has unexpectedly been shut down
 	 */
-	virtual bool PostActiveToolShutdownRequest(
+	INTERACTIVETOOLSFRAMEWORK_API virtual bool PostActiveToolShutdownRequest(
 		UInteractiveTool* Tool, 
 		EToolShutdownType ShutdownType,
 		bool bShowUnexpectedShutdownMessage = false,
@@ -284,12 +284,12 @@ public:
 	/** @return current IToolsContextTransactionsAPI */
 	virtual IToolsContextTransactionsAPI* GetContextTransactionsAPI() final { return TransactionsAPI; }
 
-	UInteractiveGizmoManager* GetPairedGizmoManager();
+	INTERACTIVETOOLSFRAMEWORK_API UInteractiveGizmoManager* GetPairedGizmoManager();
 
 	/**
 	 * @return the context object store from the owning tools context.
 	 */
-	UContextObjectStore* GetContextObjectStore() const;
+	INTERACTIVETOOLSFRAMEWORK_API UContextObjectStore* GetContextObjectStore() const;
 
 public:
 	/** Currently-active Left Tool, or null if no Tool is active */
@@ -361,8 +361,8 @@ protected:
 	 */
 	bool bActiveToolMadeSelectionStoreRequest = false;
 
-	virtual bool ActivateToolInternal(EToolSide Side);
-	virtual void DeactivateToolInternal(EToolSide Side, EToolShutdownType ShutdownType);
+	INTERACTIVETOOLSFRAMEWORK_API virtual bool ActivateToolInternal(EToolSide Side);
+	INTERACTIVETOOLSFRAMEWORK_API virtual void DeactivateToolInternal(EToolSide Side, EToolShutdownType ShutdownType);
 
 	friend class FBeginToolChange;
 	friend class FActivateToolChange;
@@ -375,23 +375,23 @@ protected:
  * FBeginToolChange is used by UInteractiveToolManager to back out of a Tool on Undo.
  * No action is taken on Redo, ie we do not re-start the Tool on Redo.
  */
-class INTERACTIVETOOLSFRAMEWORK_API FBeginToolChange : public FToolCommandChange
+class FBeginToolChange : public FToolCommandChange
 {
 public:
-	virtual void Apply(UObject* Object) override;
+	INTERACTIVETOOLSFRAMEWORK_API virtual void Apply(UObject* Object) override;
 
-	virtual void Revert(UObject* Object) override;
+	INTERACTIVETOOLSFRAMEWORK_API virtual void Revert(UObject* Object) override;
 
-	virtual bool HasExpired(UObject* Object) const override;
+	INTERACTIVETOOLSFRAMEWORK_API virtual bool HasExpired(UObject* Object) const override;
 
-	virtual FString ToString() const override;
+	INTERACTIVETOOLSFRAMEWORK_API virtual FString ToString() const override;
 };
 
 /**
  * FActivateToolChange is used by UInteractiveToolManager to change the active tool.
  * This Change has two modes, either activating or deactivating.
  */
-class INTERACTIVETOOLSFRAMEWORK_API FActivateToolChange : public FToolCommandChange
+class FActivateToolChange : public FToolCommandChange
 {
 public:
 	EToolSide Side;
@@ -404,10 +404,10 @@ public:
 	FActivateToolChange(EToolSide SideIn, FString ToolTypeIn, EToolShutdownType ShutdownTypeIn)
 		: Side(SideIn), ToolType(ToolTypeIn), bIsDeactivate(true), ShutdownType(ShutdownTypeIn) {}
 
-	virtual void Apply(UObject* Object) override;
-	virtual void Revert(UObject* Object) override;
-	virtual bool HasExpired(UObject* Object) const override;
-	virtual FString ToString() const override;
+	INTERACTIVETOOLSFRAMEWORK_API virtual void Apply(UObject* Object) override;
+	INTERACTIVETOOLSFRAMEWORK_API virtual void Revert(UObject* Object) override;
+	INTERACTIVETOOLSFRAMEWORK_API virtual bool HasExpired(UObject* Object) const override;
+	INTERACTIVETOOLSFRAMEWORK_API virtual FString ToString() const override;
 };
 
 
@@ -418,18 +418,18 @@ public:
  * FToolChangeWrapperChange wraps an FChange emitted by an InteractiveTool, allowing
  * us to Expire the change without each FChange implementation needing to handle this explicitly.
  */
-class INTERACTIVETOOLSFRAMEWORK_API FToolChangeWrapperChange : public FToolCommandChange
+class FToolChangeWrapperChange : public FToolCommandChange
 {
 public:
 	TWeakObjectPtr<UInteractiveToolManager> ToolManager;
 	TWeakObjectPtr<UInteractiveTool> ActiveTool;
 	TUniquePtr<FToolCommandChange> ToolChange;
 
-	virtual void Apply(UObject* Object) override;
+	INTERACTIVETOOLSFRAMEWORK_API virtual void Apply(UObject* Object) override;
 
-	virtual void Revert(UObject* Object) override;
+	INTERACTIVETOOLSFRAMEWORK_API virtual void Revert(UObject* Object) override;
 
-	virtual bool HasExpired(UObject* Object) const override;
+	INTERACTIVETOOLSFRAMEWORK_API virtual bool HasExpired(UObject* Object) const override;
 
-	virtual FString ToString() const override;
+	INTERACTIVETOOLSFRAMEWORK_API virtual FString ToString() const override;
 };
diff --git a/Engine/Source/Runtime/InteractiveToolsFramework/Public/InteractiveToolObjects.h b/Engine/Source/Runtime/InteractiveToolsFramework/Public/InteractiveToolObjects.h
--- a/Engine/Source/Runtime/InteractiveToolsFramework/Public/InteractiveToolObjects.h
+++ b/Engine/Source/Runtime/InteractiveToolsFramework/Public/InteractiveToolObjects.h
@@ -14,8 +14,8 @@
  * may need special-case handling, for example to prevent the user from
  * selecting and deleting them. 
  */
-UCLASS(Transient, NotPlaceable, Hidden, NotBlueprintable, NotBlueprintType)
-class INTERACTIVETOOLSFRAMEWORK_API AInternalToolFrameworkActor : public AActor
+UCLASS(Transient, NotPlaceable, Hidden, NotBlueprintable, NotBlueprintType, MinimalAPI)
+class AInternalToolFrameworkActor : public AActor
 {
 	GENERATED_BODY()
 private:
