InputState: Make available to BP

diff --git a/Engine/Source/Runtime/InteractiveToolsFramework/Public/InputState.h b/Engine/Source/Runtime/InteractiveToolsFramework/Public/InputState.h
--- a/Engine/Source/Runtime/InteractiveToolsFramework/Public/InputState.h
+++ b/Engine/Source/Runtime/InteractiveToolsFramework/Public/InputState.h
@@ -139,16 +139,25 @@ struct FInputRayHit
  * In addition there may be additional frames where the button is
  * held down and bDown=true and bPressed=false.
  */
+USTRUCT(BlueprintType)
 struct FDeviceButtonState
 {
+	GENERATED_BODY()
+	
 	/** Button identifier */
+	UPROPERTY(transient, BlueprintReadWrite, Category = DeviceButtonState)
 	FKey Button;
 
 	/** Was the button pressed down this frame. This should happen once per "click" */
+	UPROPERTY(transient, BlueprintReadWrite, Category = DeviceButtonState)
 	bool bPressed;
+
 	/** Is the button currently pressed down. This should be true every frame the button is pressed. */
+	UPROPERTY(transient, BlueprintReadWrite, Category = DeviceButtonState)
 	bool bDown;
+
 	/** Was the button released this frame. This should happen once per "click" */
+	UPROPERTY(transient, BlueprintReadWrite, Category = DeviceButtonState)
 	bool bReleased;
 
 	FDeviceButtonState()
@@ -178,9 +187,13 @@ struct FDeviceButtonState
  * Current state of active keyboard key at a point in time
  * @todo would be useful to track set of active keys
  */
+USTRUCT(BlueprintType)
 struct FKeyboardInputDeviceState
 {
+	GENERATED_BODY()
+	
 	/** state of active key that was modified (ie press or release) */
+	UPROPERTY(transient, BlueprintReadWrite, Category = KeyboardInputDeviceState)
 	FDeviceButtonState ActiveKey;
 };
 
@@ -189,28 +202,39 @@ struct FKeyboardInputDeviceState
 /**
  * Current State of a physical Mouse device at a point in time.
  */
+USTRUCT(BlueprintType)
 struct FMouseInputDeviceState
 {
+	GENERATED_BODY()
+	
 	/** State of the left mouse button */
+	UPROPERTY(transient, BlueprintReadWrite, Category = MouseInputDeviceState)
 	FDeviceButtonState Left;
+		
 	/** State of the middle mouse button */
+	UPROPERTY(transient, BlueprintReadWrite, Category = MouseInputDeviceState)
 	FDeviceButtonState Middle;
+		
 	/** State of the right mouse button */
+	UPROPERTY(transient, BlueprintReadWrite, Category = MouseInputDeviceState)
 	FDeviceButtonState Right;
 
 	/** Change in 'ticks' of the mouse wheel since last state event */
+	UPROPERTY(transient, BlueprintReadWrite, Category = MouseInputDeviceState)
 	float WheelDelta;
 
 	/** Current 2D position of the mouse, in application-defined coordinate system */
+	UPROPERTY(transient, BlueprintReadWrite, Category = MouseInputDeviceState)
 	FVector2D Position2D;
 
 	/** Change in 2D mouse position from last state event */
+	UPROPERTY(transient, BlueprintReadWrite, Category = MouseInputDeviceState)
 	FVector2D Delta2D;
 
 	/** Ray into current 3D scene at current 2D mouse position */
+	UPROPERTY(transient, BlueprintReadWrite, Category = MouseInputDeviceState)
 	FRay WorldRay;
 
-
 	FMouseInputDeviceState()
 	{
 		Left = FDeviceButtonState(EKeys::LeftMouseButton);
@@ -230,9 +254,13 @@ struct FMouseInputDeviceState
  * Assumption is that the state refers to a single physical input device,
  * ie InputDevice field is a single value of EInputDevices and not a combination.
  */
+USTRUCT(BlueprintType)
 struct FInputDeviceState
 {
+	GENERATED_BODY()
+
 	/** Which InputDevice member is valid in this state */
+	UPROPERTY(transient, BlueprintReadWrite, Category = MouseInputDeviceState)
 	EInputDevices InputDevice;
 
 	//
@@ -240,22 +268,29 @@ struct FInputDeviceState
 	//
 
 	/** Is they keyboard SHIFT modifier key currently pressed down */
+	UPROPERTY(transient, BlueprintReadWrite, Category = MouseInputDeviceState)
 	bool bShiftKeyDown;
+	
 	/** Is they keyboard ALT modifier key currently pressed down */
+	UPROPERTY(transient, BlueprintReadWrite, Category = MouseInputDeviceState)
 	bool bAltKeyDown;
+	
 	/** Is they keyboard CTRL modifier key currently pressed down */
+	UPROPERTY(transient, BlueprintReadWrite, Category = MouseInputDeviceState)
 	bool bCtrlKeyDown;
+	
 	/** Is they keyboard CMD modifier key currently pressed down (only on Apple devices) */
+	UPROPERTY(transient, BlueprintReadWrite, Category = MouseInputDeviceState)
 	bool bCmdKeyDown;
 
-
 	/** Current state of Keyboard device, if InputDevice == EInputDevices::Keyboard */
+	UPROPERTY(transient, BlueprintReadWrite, Category = MouseInputDeviceState)
 	FKeyboardInputDeviceState Keyboard;
 
 	/** Current state of Mouse device, if InputDevice == EInputDevices::Mouse */
+	UPROPERTY(transient, BlueprintReadWrite, Category = MouseInputDeviceState)
 	FMouseInputDeviceState Mouse;
 
-
 	FInputDeviceState() 
 	{
 		InputDevice = EInputDevices::None;
@@ -328,15 +363,15 @@ struct FInputDeviceRay
 	GENERATED_BODY()
 
 	/** 3D ray in 3D scene, in world coordinates */
-	UPROPERTY(BlueprintReadWrite, Category = InputDeviceRay)
+	UPROPERTY(transient, BlueprintReadWrite, Category = InputDeviceRay)
 	FRay WorldRay;
 
 	/** If true, WorldRay has 2D device position coordinates */
-	UPROPERTY(BlueprintReadWrite, Category = InputDeviceRay)
+	UPROPERTY(transient, BlueprintReadWrite, Category = InputDeviceRay)
 	bool bHas2D = false;
 
 	/** 2D device position coordinates associated with the ray */
-	UPROPERTY(BlueprintReadWrite, Category = InputDeviceRay)
+	UPROPERTY(transient, BlueprintReadWrite, Category = InputDeviceRay)
 	FVector2D ScreenPosition;
 
 	// this is required for a USTRUCT
