ITF: Sphere, LineStrip and TriangleList elements

diff --git a/Engine/Source/Runtime/InteractiveToolsFramework/Private/BaseGizmos/GizmoElementLineStrip.cpp b/Engine/Source/Runtime/InteractiveToolsFramework/Private/BaseGizmos/GizmoElementLineStrip.cpp
new file mode 100644
--- /dev/null
+++ b/Engine/Source/Runtime/InteractiveToolsFramework/Private/BaseGizmos/GizmoElementLineStrip.cpp
@@ -0,0 +1,180 @@
+// Copyright Epic Games, Inc. All Rights Reserved.
+ 
+#include "BaseGizmos/GizmoElementLineStrip.h"
+#include "BaseGizmos/GizmoRenderingUtil.h"
+#include "BaseGizmos/GizmoMath.h"
+#include "InputState.h"
+#include "Materials/MaterialInterface.h"
+#include "Math/Ray.h"
+#include "SceneManagement.h"
+#include "Math/UnrealMathUtility.h"
+
+#include UE_INLINE_GENERATED_CPP_BY_NAME(GizmoElementLineStrip)
+
+void UGizmoElementLineStrip::Render(IToolsContextRenderAPI* RenderAPI, const FRenderTraversalState& RenderState)
+{
+	FRenderTraversalState CurrentRenderState(RenderState);
+	const bool bVisibleViewDependent = UpdateRenderState(RenderAPI, Base, CurrentRenderState);
+
+	if (bVisibleViewDependent && Vertices.Num() > 1)
+	{
+		FPrimitiveDrawInterface* PDI = RenderAPI->GetPrimitiveDrawInterface();
+		ComputeProjectedVertices(RenderState.LocalToWorldTransform, true);
+
+		check(RenderAPI);
+		const FSceneView* View = RenderAPI->GetSceneView();
+		check(View);
+
+		const float CurrentLineThickness = GetCurrentLineThickness(View->IsPerspectiveProjection(), View->FOV);
+		const FColor LineColor = CurrentRenderState.GetCurrentLineColor().ToFColor(false);
+		
+		if(bDrawLineStrip)
+		{
+			for(int32 Index = 0; Index < ProjectedVertices.Num() - 1; Index++)
+			{
+				const FVector& A = ProjectedVertices[Index];
+				const FVector& B = ProjectedVertices[Index+1];
+				PDI->DrawLine(A, B, LineColor, SDPG_Foreground, CurrentLineThickness, 0.0, bScreenSpaceLine);
+			}
+		}
+		else
+		{
+			for(int32 Index = 0; Index < ProjectedVertices.Num(); Index+=2)
+			{
+				const FVector& A = ProjectedVertices[Index];
+				const FVector& B = ProjectedVertices[Index+1];
+				PDI->DrawLine(A, B, LineColor, SDPG_Foreground, CurrentLineThickness, 0.0, bScreenSpaceLine);
+			}
+		}
+	}
+}
+
+FInputRayHit UGizmoElementLineStrip::LineTrace(const UGizmoViewContext* ViewContext, const FLineTraceTraversalState& LineTraceState, const FVector& RayOrigin, const FVector& RayDirection)
+{
+	FLineTraceTraversalState CurrentLineTraceState(LineTraceState);
+	const bool bHittableViewDependent = UpdateLineTraceState(ViewContext, Base, CurrentLineTraceState);
+
+	if (bHittableViewDependent && Vertices.Num() > 1)
+	{
+		const double PixelHitThresholdAdjust = CurrentLineTraceState.PixelToWorldScale * PixelHitDistanceThreshold;
+		ComputeProjectedVertices(CurrentLineTraceState.LocalToWorldTransform, false);
+
+		const FRay Ray(RayOrigin, RayDirection);
+		const FVector RayB = Ray.Origin + Ray.Direction * 10000.0;
+		auto Intersect = [Ray, RayB](const FVector& A, const FVector&B, double InTolerance, double& Distance) -> bool
+		{
+			FVector HitPointA = Ray.Origin;
+			FVector HitPointB = A;
+			FMath::SegmentDistToSegmentSafe(Ray.Origin, RayB, A, B, HitPointA, HitPointB);
+			if(HitPointA.Equals(Ray.Origin))
+			{
+				return false;
+			}
+
+			Distance = FVector::Distance(HitPointA, HitPointB);
+			return Distance < InTolerance;
+		};
+
+		double HitDepth = 0;
+		if(bDrawLineStrip)
+		{
+			for(int32 Index = 0; Index < ProjectedVertices.Num() - 1; Index++)
+			{
+				const FVector& A = ProjectedVertices[Index];
+				const FVector& B = ProjectedVertices[Index+1];
+				if(Intersect(A, B, PixelHitThresholdAdjust, HitDepth))
+				{
+					FInputRayHit RayHit(static_cast<float>(HitDepth));
+					RayHit.SetHitObject(this);
+					RayHit.HitIdentifier = PartIdentifier;
+					return RayHit;
+				}
+			}
+		}
+		else
+		{
+			for(int32 Index = 0; Index < ProjectedVertices.Num(); Index+=2)
+			{
+				const FVector& A = ProjectedVertices[Index];
+				const FVector& B = ProjectedVertices[Index+1];
+				if(Intersect(A, B, PixelHitThresholdAdjust, HitDepth))
+				{
+					FInputRayHit RayHit(static_cast<float>(HitDepth));
+					RayHit.SetHitObject(this);
+					RayHit.HitIdentifier = PartIdentifier;
+					return RayHit;
+				}
+			}
+		}
+	}
+	return FInputRayHit();
+}
+
+void UGizmoElementLineStrip::SetVertices(const TArrayView<const FVector>& InVertices)
+{
+	Vertices = InVertices;
+}
+
+const TArray<FVector>& UGizmoElementLineStrip::GetVertices() const
+{
+	return Vertices;
+}
+
+void UGizmoElementLineStrip::SetBase(FVector InBase)
+{
+	Base = InBase;
+}
+
+FVector UGizmoElementLineStrip::GetBase() const
+{
+	return Base;
+}
+
+void UGizmoElementLineStrip::SetUpDirection(const FVector& InUpDirection)
+{
+	UpDirection = InUpDirection.GetSafeNormal();
+}
+
+FVector UGizmoElementLineStrip::GetUpDirection() const
+{
+	return UpDirection;
+}
+
+void UGizmoElementLineStrip::SetSideDirection(const FVector& InSideDirection)
+{
+	SideDirection = InSideDirection.GetSafeNormal();
+}
+
+FVector UGizmoElementLineStrip::GetSideDirection() const
+{
+	return SideDirection;
+}
+
+void UGizmoElementLineStrip::SetDrawLineStrip(bool InDrawLineStrip)
+{
+	bDrawLineStrip = InDrawLineStrip;
+}
+
+bool UGizmoElementLineStrip::GetDrawLineStrip() const
+{
+	return bDrawLineStrip;
+}
+
+void UGizmoElementLineStrip::ComputeProjectedVertices(const FTransform& InLocalToWorldTransform, bool bIncludeBase)
+{
+	const FVector WorldUpAxis = InLocalToWorldTransform.TransformVectorNoScale(UpDirection);
+    const FVector WorldSideAxis = InLocalToWorldTransform.TransformVectorNoScale(SideDirection);
+    const FVector WorldNormal = FVector::CrossProduct(WorldSideAxis, WorldUpAxis);
+    const FVector WorldCenter = InLocalToWorldTransform.TransformPosition(bIncludeBase ? Base : FVector::ZeroVector);
+    const FVector WorldScale = InLocalToWorldTransform.GetScale3D();
+
+    ProjectedVertices.Reset();
+    ProjectedVertices.Reserve(Vertices.Num());
+
+    for(int32 Index = 0; Index < Vertices.Num(); Index++)
+    {
+    	FVector Vertex = Vertices[Index] * WorldScale;
+    	ProjectedVertices.Add(WorldCenter + WorldNormal * Vertex.X + WorldSideAxis * Vertex.Y + WorldUpAxis * Vertex.Z);
+    }
+}
+
diff --git a/Engine/Source/Runtime/InteractiveToolsFramework/Private/BaseGizmos/GizmoElementSphere.cpp b/Engine/Source/Runtime/InteractiveToolsFramework/Private/BaseGizmos/GizmoElementSphere.cpp
new file mode 100644
--- /dev/null
+++ b/Engine/Source/Runtime/InteractiveToolsFramework/Private/BaseGizmos/GizmoElementSphere.cpp
@@ -0,0 +1,90 @@
+// Copyright Epic Games, Inc. All Rights Reserved.
+ 
+#include "BaseGizmos/GizmoElementSphere.h"
+#include "BaseGizmos/GizmoRenderingUtil.h"
+#include "Materials/MaterialInterface.h"
+#include "InputState.h"
+#include "Math/Ray.h"
+#include "SceneManagement.h"
+
+#include UE_INLINE_GENERATED_CPP_BY_NAME(GizmoElementSphere)
+
+void UGizmoElementSphere::Render(IToolsContextRenderAPI* RenderAPI, const FRenderTraversalState& RenderState)
+{
+	FRenderTraversalState CurrentRenderState(RenderState);
+	bool bVisibleViewDependent = UpdateRenderState(RenderAPI, Center, CurrentRenderState);
+
+	if (bVisibleViewDependent)
+	{
+		if (const UMaterialInterface* UseMaterial = CurrentRenderState.GetCurrentMaterial())
+		{
+			FPrimitiveDrawInterface* PDI = RenderAPI->GetPrimitiveDrawInterface();
+			DrawSphere(
+				PDI,
+				CurrentRenderState.LocalToWorldTransform.GetTranslation(),
+				CurrentRenderState.LocalToWorldTransform.Rotator(),
+				FVector::OneVector * Radius,
+				NumSides,
+				NumSides,
+				UseMaterial->GetRenderProxy(), SDPG_Foreground);
+		}
+	}
+}
+
+FInputRayHit UGizmoElementSphere::LineTrace(const UGizmoViewContext* ViewContext, const FLineTraceTraversalState& LineTraceState, const FVector& RayOrigin, const FVector& RayDirection)
+{
+	FLineTraceTraversalState CurrentLineTraceState(LineTraceState);
+	bool bHittableViewDependent = UpdateLineTraceState(ViewContext, Center, CurrentLineTraceState);
+
+	if (bHittableViewDependent)
+	{
+		const FVector WorldCenter = CurrentLineTraceState.LocalToWorldTransform.GetTranslation();
+		const double PixelHitThresholdAdjust = CurrentLineTraceState.PixelToWorldScale * PixelHitDistanceThreshold;
+		const double Scale = CurrentLineTraceState.LocalToWorldTransform.GetScale3D().X;
+		const double WorldExtent = Scale * Radius + PixelHitThresholdAdjust;
+
+		const FRay Ray(RayOrigin, RayDirection);
+		const double HitDepth = Ray.Dist(WorldCenter);
+
+		if(HitDepth < WorldExtent)
+		{
+			FInputRayHit RayHit(static_cast<float>(HitDepth));
+			RayHit.SetHitObject(this);
+			RayHit.HitIdentifier = PartIdentifier;
+			return RayHit;
+		}
+	}
+
+	return FInputRayHit();
+}
+
+void UGizmoElementSphere::SetCenter(const FVector& InCenter)
+{
+	Center = InCenter;
+}
+
+FVector UGizmoElementSphere::GetCenter() const
+{
+	return Center;
+}
+
+void UGizmoElementSphere::SetRadius(float InRadius)
+{
+	Radius = InRadius;
+}
+
+float UGizmoElementSphere::GetRadius() const
+{
+	return Radius;
+}
+
+void UGizmoElementSphere::SetNumSides(int32 InNumSides)
+{
+	NumSides = InNumSides;
+}
+
+int32 UGizmoElementSphere::GetNumSides() const
+{
+	return NumSides;
+}
+
diff --git a/Engine/Source/Runtime/InteractiveToolsFramework/Private/BaseGizmos/GizmoElementTriangleList.cpp b/Engine/Source/Runtime/InteractiveToolsFramework/Private/BaseGizmos/GizmoElementTriangleList.cpp
new file mode 100644
--- /dev/null
+++ b/Engine/Source/Runtime/InteractiveToolsFramework/Private/BaseGizmos/GizmoElementTriangleList.cpp
@@ -0,0 +1,154 @@
+// Copyright Epic Games, Inc. All Rights Reserved.
+ 
+#include "BaseGizmos/GizmoElementTriangleList.h"
+#include "BaseGizmos/GizmoRenderingUtil.h"
+#include "Materials/MaterialInterface.h"
+#include "InputState.h"
+#include "Math/Ray.h"
+#include "SceneManagement.h"
+#include "DynamicMeshBuilder.h"
+
+#include UE_INLINE_GENERATED_CPP_BY_NAME(GizmoElementTriangleList)
+
+void UGizmoElementTriangleList::Render(IToolsContextRenderAPI* RenderAPI, const FRenderTraversalState& RenderState)
+{
+	FRenderTraversalState CurrentRenderState(RenderState);
+	const bool bVisibleViewDependent = UpdateRenderState(RenderAPI, Base, CurrentRenderState);
+
+	if (bVisibleViewDependent && Vertices.Num() > 2)
+	{
+		if (const UMaterialInterface* UseMaterial = CurrentRenderState.GetCurrentMaterial())
+		{
+			FPrimitiveDrawInterface* PDI = RenderAPI->GetPrimitiveDrawInterface();
+			ComputeProjectedVertices(RenderState.LocalToWorldTransform, true);
+
+			FDynamicMeshBuilder MeshBuilder(PDI->View->GetFeatureLevel());
+			for(int32 Index = 0; Index < ProjectedVertices.Num() - 2;)
+			{
+				const FDynamicMeshVertex A((FVector3f)ProjectedVertices[Index++]);
+				const FDynamicMeshVertex B((FVector3f)ProjectedVertices[Index++]);
+				const FDynamicMeshVertex C((FVector3f)ProjectedVertices[Index++]);
+				MeshBuilder.AddTriangle(MeshBuilder.AddVertex(A), MeshBuilder.AddVertex(B), MeshBuilder.AddVertex(C));
+			}
+			MeshBuilder.Draw(PDI,FMatrix::Identity,UseMaterial->GetRenderProxy(),SDPG_Foreground,0.f);
+		}
+	}
+}
+
+FInputRayHit UGizmoElementTriangleList::LineTrace(const UGizmoViewContext* ViewContext, const FLineTraceTraversalState& LineTraceState, const FVector& RayOrigin, const FVector& RayDirection)
+{
+	FLineTraceTraversalState CurrentLineTraceState(LineTraceState);
+	const bool bHittableViewDependent = UpdateLineTraceState(ViewContext, Base, CurrentLineTraceState);
+
+	if (bHittableViewDependent && Vertices.Num() > 2)
+	{
+		ComputeProjectedVertices(LineTraceState.LocalToWorldTransform, false);
+		const double PixelHitThresholdAdjust = CurrentLineTraceState.PixelToWorldScale * PixelHitDistanceThreshold;
+
+		const FRay Ray(RayOrigin, RayDirection);
+
+		double HitDepth = 0.0;
+		for(int32 Index = 0; Index < ProjectedVertices.Num() - 2;)
+		{
+			const FVector& A = ProjectedVertices[Index++];
+			const FVector& B = ProjectedVertices[Index++];
+			const FVector& C = ProjectedVertices[Index++];
+			if (Intersect(Ray, A, B, C, PixelHitThresholdAdjust, HitDepth))
+			{
+				FInputRayHit RayHit(static_cast<float>(HitDepth));
+				RayHit.SetHitObject(this);
+				RayHit.HitIdentifier = PartIdentifier;
+				return RayHit;
+			}
+		}
+	}
+
+	return FInputRayHit();
+}
+
+void UGizmoElementTriangleList::SetVertices(const TArrayView<const FVector>& InVertices)
+{
+	Vertices = InVertices;
+}
+
+const TArray<FVector>& UGizmoElementTriangleList::GetVertices() const
+{
+	return Vertices;
+}
+
+void UGizmoElementTriangleList::SetBase(FVector InBase)
+{
+	Base = InBase;
+}
+
+FVector UGizmoElementTriangleList::GetBase() const
+{
+	return Base;
+}
+
+void UGizmoElementTriangleList::SetUpDirection(const FVector& InUpDirection)
+{
+	UpDirection = InUpDirection.GetSafeNormal();
+}
+
+FVector UGizmoElementTriangleList::GetUpDirection() const
+{
+	return UpDirection;
+}
+
+void UGizmoElementTriangleList::SetSideDirection(const FVector& InSideDirection)
+{
+	SideDirection = InSideDirection.GetSafeNormal();
+}
+
+FVector UGizmoElementTriangleList::GetSideDirection() const
+{
+	return SideDirection;
+}
+
+void UGizmoElementTriangleList::ComputeProjectedVertices(const FTransform& InLocalToWorldTransform, bool bIncludeBase)
+{
+	const FVector WorldUpAxis = InLocalToWorldTransform.TransformVectorNoScale(UpDirection);
+	const FVector WorldSideAxis = InLocalToWorldTransform.TransformVectorNoScale(SideDirection);
+    const FVector WorldNormal = FVector::CrossProduct(WorldSideAxis, WorldUpAxis);
+	const FVector WorldCenter = InLocalToWorldTransform.TransformPosition(Base);
+	const FVector WorldScale = InLocalToWorldTransform.GetScale3D();
+
+	ProjectedVertices.Reset();
+	ProjectedVertices.Reserve(Vertices.Num());
+
+	for(int32 Index = 0; Index < Vertices.Num(); Index++)
+	{
+		FVector Vertex = Vertices[Index] * WorldScale;
+		ProjectedVertices.Add(WorldCenter + WorldNormal * Vertex.X + WorldSideAxis * Vertex.Y + WorldUpAxis * Vertex.Z);
+	}
+}
+
+bool UGizmoElementTriangleList::Intersect(const FRay& Ray, const FVector& InA, const FVector& InB, const FVector& InC, double Tolerance, double& Distance)
+{
+	const FVector Center = (InA + InB + InC) / 3.0;
+	const FVector A = InA + (InA - Center).GetSafeNormal() * Tolerance;
+	const FVector B = InB + (InB - Center).GetSafeNormal() * Tolerance;
+	const FVector C = InC + (InC - Center).GetSafeNormal() * Tolerance;
+	const FVector BA = A - B;
+	const FVector CB = B - C;
+	const FVector TriNormal = BA ^ CB;
+
+	FVector Intersection;
+	bool bCollide = FMath::SegmentPlaneIntersection(
+		Ray.Origin,
+		Ray.Origin + Ray.Direction * 10000.0,
+		FPlane(A, TriNormal), Intersection);
+	if (!bCollide)
+	{
+		return false;
+	}
+
+	const FVector BaryCentric = FMath::ComputeBaryCentric2D(Intersection, A, B, C);
+	if (BaryCentric.X > 0.0f && BaryCentric.Y > 0.0f && BaryCentric.Z > 0.0f)
+	{
+		Distance = 0.0;
+		return true;
+	}
+	return false;
+}
\ No newline at end of file
diff --git a/Engine/Source/Runtime/InteractiveToolsFramework/Public/BaseGizmos/GizmoElementLineStrip.h b/Engine/Source/Runtime/InteractiveToolsFramework/Public/BaseGizmos/GizmoElementLineStrip.h
new file mode 100644
--- /dev/null
+++ b/Engine/Source/Runtime/InteractiveToolsFramework/Public/BaseGizmos/GizmoElementLineStrip.h
@@ -0,0 +1,75 @@
+// Copyright Epic Games, Inc. All Rights Reserved.
+
+#pragma once
+
+#include "CoreMinimal.h"
+#include "BaseGizmos/GizmoElementLineBase.h"
+#include "InputState.h"
+#include "UObject/ObjectMacros.h"
+#include "GizmoElementLineStrip.generated.h"
+
+class FPrimitiveDrawInterface;
+class FMaterialRenderProxy;
+
+/**
+ * Simple object intended to be used as part of 3D Gizmos.
+ * Draws a rectangle based on parameters.
+ */
+UCLASS(Transient, MinimalAPI)
+class UGizmoElementLineStrip : public UGizmoElementLineBase
+{
+	GENERATED_BODY()
+
+public:
+
+	//~ Begin UGizmoElementBase Interface.
+	INTERACTIVETOOLSFRAMEWORK_API virtual void Render(IToolsContextRenderAPI* RenderAPI, const FRenderTraversalState& RenderState) override;
+	INTERACTIVETOOLSFRAMEWORK_API virtual FInputRayHit LineTrace(const UGizmoViewContext* ViewContext, const FLineTraceTraversalState& LineTraceState, const FVector& RayOrigin, const FVector& RayDirection) override;
+	//~ End UGizmoElementBase Interface.
+
+	// The vertices of the line strip
+	INTERACTIVETOOLSFRAMEWORK_API virtual void SetVertices(const TArrayView<const FVector>& InVertices);
+	INTERACTIVETOOLSFRAMEWORK_API virtual const TArray<FVector>& GetVertices() const;
+
+	// Base Location of the line strip
+	INTERACTIVETOOLSFRAMEWORK_API virtual void SetBase(FVector InBase);
+	INTERACTIVETOOLSFRAMEWORK_API virtual FVector GetBase() const;
+
+	// Up direction
+	INTERACTIVETOOLSFRAMEWORK_API virtual void SetUpDirection(const FVector& InUpDirection);
+	INTERACTIVETOOLSFRAMEWORK_API virtual FVector GetUpDirection() const;
+
+	// Side direction
+	INTERACTIVETOOLSFRAMEWORK_API virtual void SetSideDirection(const FVector& InSideDirection);
+	INTERACTIVETOOLSFRAMEWORK_API virtual FVector GetSideDirection() const;
+
+	// Defines if the vertices should be treated as a connected strip or separate line segments
+	INTERACTIVETOOLSFRAMEWORK_API virtual void SetDrawLineStrip(bool InDrawLineStrip);
+	INTERACTIVETOOLSFRAMEWORK_API virtual bool GetDrawLineStrip() const;
+
+protected:
+
+	void ComputeProjectedVertices(const FTransform& InLocalToWorldTransform, bool bIncludeBase);
+
+	// The vertices of the line strip
+	UPROPERTY()
+	TArray<FVector> Vertices;
+
+	// Base Location of the line strip
+	UPROPERTY()
+	FVector Base = FVector::ZeroVector;
+
+	// Up direction
+	UPROPERTY()
+	FVector UpDirection = FVector(0.0f, 0.0f, 1.0f);
+
+	// Side direction
+	UPROPERTY()
+	FVector SideDirection = FVector(0.0f, 1.0f, 0.0f);
+
+	// Defines if the vertices should be treated as a connected strip or separate line segments
+	UPROPERTY()
+	bool bDrawLineStrip = true;
+
+	TArray<FVector> ProjectedVertices;
+};
diff --git a/Engine/Source/Runtime/InteractiveToolsFramework/Public/BaseGizmos/GizmoElementShapes.h b/Engine/Source/Runtime/InteractiveToolsFramework/Public/BaseGizmos/GizmoElementShapes.h
--- a/Engine/Source/Runtime/InteractiveToolsFramework/Public/BaseGizmos/GizmoElementShapes.h
+++ b/Engine/Source/Runtime/InteractiveToolsFramework/Public/BaseGizmos/GizmoElementShapes.h
@@ -9,6 +9,8 @@
 #include "BaseGizmos/GizmoElementCircle.h"
 #include "BaseGizmos/GizmoElementCylinder.h"
 #include "BaseGizmos/GizmoElementGroup.h"
+#include "BaseGizmos/GizmoElementLineStrip.h"
 #include "BaseGizmos/GizmoElementRectangle.h"
+#include "BaseGizmos/GizmoElementSphere.h"
 #include "BaseGizmos/GizmoElementTorus.h"
-
+#include "BaseGizmos/GizmoElementTriangleList.h"
diff --git a/Engine/Source/Runtime/InteractiveToolsFramework/Public/BaseGizmos/GizmoElementSphere.h b/Engine/Source/Runtime/InteractiveToolsFramework/Public/BaseGizmos/GizmoElementSphere.h
new file mode 100644
--- /dev/null
+++ b/Engine/Source/Runtime/InteractiveToolsFramework/Public/BaseGizmos/GizmoElementSphere.h
@@ -0,0 +1,54 @@
+// Copyright Epic Games, Inc. All Rights Reserved.
+
+#pragma once
+
+#include "CoreMinimal.h"
+#include "BaseGizmos/GizmoElementBase.h"
+#include "InputState.h"
+#include "UObject/ObjectMacros.h"
+#include "GizmoElementSphere.generated.h"
+
+/**
+ * Simple object intended to be used as part of 3D Gizmos.
+ * Draws a solid 3D sphere based on parameters.
+ */
+
+UCLASS(Transient, MinimalAPI)
+class UGizmoElementSphere : public UGizmoElementBase
+{
+	GENERATED_BODY()
+
+public:
+
+	//~ Begin UGizmoElementBase Interface.
+	INTERACTIVETOOLSFRAMEWORK_API virtual void Render(IToolsContextRenderAPI* RenderAPI, const FRenderTraversalState& RenderState) override;
+	INTERACTIVETOOLSFRAMEWORK_API virtual FInputRayHit LineTrace(const UGizmoViewContext* ViewContext, const FLineTraceTraversalState& LineTraceState, const FVector& RayOrigin, const FVector& RayDirection) override;
+	//~ End UGizmoElementBase Interface.
+
+	// Location of center of sphere
+	INTERACTIVETOOLSFRAMEWORK_API virtual void SetCenter(const FVector& InCenter);
+	INTERACTIVETOOLSFRAMEWORK_API virtual FVector GetCenter() const;
+
+	// Sphere radius.
+	INTERACTIVETOOLSFRAMEWORK_API virtual void SetRadius(float InRadius);
+	INTERACTIVETOOLSFRAMEWORK_API virtual float GetRadius() const;
+
+	// Number of sides for tessellating sphere.
+	INTERACTIVETOOLSFRAMEWORK_API virtual void SetNumSides(int32 InNumSides);
+	INTERACTIVETOOLSFRAMEWORK_API virtual int32 GetNumSides() const;
+
+protected:
+
+	// Location of center of sphere
+	UPROPERTY()
+	FVector Center = FVector::ZeroVector;
+
+	// Sphere radius.
+	UPROPERTY()
+	float Radius = 0.5f;
+
+	// Number of sides for tessellating sphere.
+	UPROPERTY()
+	int32 NumSides = 32;
+};
+
diff --git a/Engine/Source/Runtime/InteractiveToolsFramework/Public/BaseGizmos/GizmoElementTriangleList.h b/Engine/Source/Runtime/InteractiveToolsFramework/Public/BaseGizmos/GizmoElementTriangleList.h
new file mode 100644
--- /dev/null
+++ b/Engine/Source/Runtime/InteractiveToolsFramework/Public/BaseGizmos/GizmoElementTriangleList.h
@@ -0,0 +1,68 @@
+// Copyright Epic Games, Inc. All Rights Reserved.
+
+#pragma once
+
+#include "CoreMinimal.h"
+#include "BaseGizmos/GizmoElementBase.h"
+#include "InputState.h"
+#include "UObject/ObjectMacros.h"
+#include "GizmoElementTriangleList.generated.h"
+
+/**
+ * Simple object intended to be used as part of 3D Gizmos.
+ * Draws a solid 3D sphere based on parameters.
+ */
+
+UCLASS(Transient, MinimalAPI)
+class UGizmoElementTriangleList : public UGizmoElementBase
+{
+	GENERATED_BODY()
+
+public:
+
+	//~ Begin UGizmoElementBase Interface.
+	INTERACTIVETOOLSFRAMEWORK_API virtual void Render(IToolsContextRenderAPI* RenderAPI, const FRenderTraversalState& RenderState) override;
+	INTERACTIVETOOLSFRAMEWORK_API virtual FInputRayHit LineTrace(const UGizmoViewContext* ViewContext, const FLineTraceTraversalState& LineTraceState, const FVector& RayOrigin, const FVector& RayDirection) override;
+	//~ End UGizmoElementBase Interface.
+
+	// The vertices of the triangle list, each tuple of 3 forming a triangle
+	INTERACTIVETOOLSFRAMEWORK_API virtual void SetVertices(const TArrayView<const FVector>& InVertices);
+	INTERACTIVETOOLSFRAMEWORK_API virtual const TArray<FVector>& GetVertices() const;
+
+	// Base location the triangles
+	INTERACTIVETOOLSFRAMEWORK_API virtual void SetBase(FVector InBase);
+	INTERACTIVETOOLSFRAMEWORK_API virtual FVector GetBase() const;
+
+	// Up direction
+	INTERACTIVETOOLSFRAMEWORK_API virtual void SetUpDirection(const FVector& InUpDirection);
+	INTERACTIVETOOLSFRAMEWORK_API virtual FVector GetUpDirection() const;
+
+	// Side direction
+	INTERACTIVETOOLSFRAMEWORK_API virtual void SetSideDirection(const FVector& InSideDirection);
+	INTERACTIVETOOLSFRAMEWORK_API virtual FVector GetSideDirection() const;
+
+protected:
+
+	void ComputeProjectedVertices(const FTransform& InLocalToWorldTransform, bool bIncludeBase);
+
+	static bool Intersect(const FRay& Ray, const FVector& A, const FVector& B, const FVector& C, double Tolerance, double& Distance);
+
+	// The vertices of the triangle list, each tuple of 3 forming a triangle
+	UPROPERTY()
+	TArray<FVector> Vertices;
+
+	// Base location the triangles
+	UPROPERTY()
+	FVector Base = FVector::ZeroVector;
+
+	// Up direction
+	UPROPERTY()
+	FVector UpDirection = FVector(0.0f, 0.0f, 1.0f);
+
+	// Side direction
+	UPROPERTY()
+	FVector SideDirection = FVector(0.0f, 1.0f, 0.0f);
+
+	TArray<FVector> ProjectedVertices;
+};
+
