InteractiveToolsFramework: Specify underlying type of uint8 for EToolContextCoordinateSystem.

#rnx

diff --git a/Engine/Source/Runtime/InteractiveToolsFramework/Public/ToolContextInterfaces.h b/Engine/Source/Runtime/InteractiveToolsFramework/Public/ToolContextInterfaces.h
--- a/Engine/Source/Runtime/InteractiveToolsFramework/Public/ToolContextInterfaces.h
+++ b/Engine/Source/Runtime/InteractiveToolsFramework/Public/ToolContextInterfaces.h
@@ -122,7 +122,7 @@ enum class EStandardToolContextMaterials
 
 /** Types of coordinate systems that a Tool/Gizmo might use */
 UENUM()
-enum class EToolContextCoordinateSystem
+enum class EToolContextCoordinateSystem : uint8
 {
 	/** World space coordinate system */
 	World = 0,
