InteractiveToolsFramework: Fix hover behaviors terminating when a modifier key is pressed.

diff --git a/Engine/Source/Runtime/InteractiveToolsFramework/Private/BaseBehaviors/MouseHoverBehavior.cpp b/Engine/Source/Runtime/InteractiveToolsFramework/Private/BaseBehaviors/MouseHoverBehavior.cpp
--- a/Engine/Source/Runtime/InteractiveToolsFramework/Private/BaseBehaviors/MouseHoverBehavior.cpp
+++ b/Engine/Source/Runtime/InteractiveToolsFramework/Private/BaseBehaviors/MouseHoverBehavior.cpp
@@ -53,6 +53,13 @@ FInputCaptureUpdate UMouseHoverBehavior::UpdateHoverCapture(const FInputDeviceSt
 {
 	check(Target != nullptr);
 	Modifiers.UpdateModifiers(InputState, Target);
+
+	// Check device here because we may get keyboard inputs for updating the modifier, and they don't have a valid ray.
+	if ((GetSupportedDevices() & InputState.InputDevice) == EInputDevices::None)
+	{
+		return FInputCaptureUpdate::Continue();
+	}
+
 	if (Target->OnUpdateHover(FInputDeviceRay(InputState)))
 	{
 		return FInputCaptureUpdate::Continue();
diff --git a/Engine/Source/Runtime/InteractiveToolsFramework/Private/InputRouter.cpp b/Engine/Source/Runtime/InteractiveToolsFramework/Private/InputRouter.cpp
--- a/Engine/Source/Runtime/InteractiveToolsFramework/Private/InputRouter.cpp
+++ b/Engine/Source/Runtime/InteractiveToolsFramework/Private/InputRouter.cpp
@@ -87,7 +87,7 @@ bool UInputRouter::PostInputEvent(const FInputDeviceState& Input)
 				}
 				else
 				{
-					bool bHoverStateUpdated = ProcessMouseHover(Input);
+					bool bHoverStateUpdated = UpdateExistingHoverCaptureIfPresent(Input);
 					if (bHoverStateUpdated && bAutoInvalidateOnHover)
 					{
 						TransactionsAPI->PostInvalidation();
@@ -306,6 +306,21 @@ void UInputRouter::TerminateHover(EInputCaptureSide Side)
 	}
 }
 
+// Returns true if hover state is updated
+bool UInputRouter::UpdateExistingHoverCaptureIfPresent(const FInputDeviceState& Input)
+{
+	if (ActiveLeftHoverCapture != nullptr)
+	{
+		FInputCaptureUpdate Result = ActiveLeftHoverCapture->UpdateHoverCapture(Input);
+		if (Result.State == EInputCaptureState::End)
+		{
+			TerminateHover(EInputCaptureSide::Left);
+			return true;
+		}
+	}
+
+	return false;
+}
 
 bool UInputRouter::ProcessMouseHover(const FInputDeviceState& Input)
 {
@@ -322,52 +337,42 @@ bool UInputRouter::ProcessMouseHover(const FInputDeviceState& Input)
 		return false;
 	}
 
+	UInputBehavior* PreviousCapture = ActiveLeftHoverCapture;
 	CaptureRequests.StableSort();
 
 	// if we have an active hover, either update it, or terminate if we got a new best hit
-	bool bHoverStateModified = false;
-	if (ActiveLeftHoverCapture != nullptr)
+	if (CaptureRequests[0].Source == ActiveLeftHoverCapture)
 	{
-		bool bTerminateActiveHover = false;
-		if (CaptureRequests[0].Source == ActiveLeftHoverCapture)
-		{
-			FInputCaptureUpdate Result =
-				ActiveLeftHoverCapture->UpdateHoverCapture(Input);
-			bTerminateActiveHover = (Result.State == EInputCaptureState::End);
-		}
-		else
-		{
-			bTerminateActiveHover = true;
-		}
-
-		if (bTerminateActiveHover)
-		{
-			TerminateHover(EInputCaptureSide::Left);
-			bHoverStateModified = true;
-		}
-		else
-		{
-			return true;		// hover has been consumed
-		}
+		UpdateExistingHoverCaptureIfPresent(Input);
+		// We don't return early because the update may have ended the hover, so we may need a replacement. 
+	}
+	else
+	{
+		TerminateHover(EInputCaptureSide::Left); // does nothing if no capture present
 	}
 
-	// if we get here, we have a new hover
-	bool bAccepted = false;
-	for (int i = 0; i < CaptureRequests.Num() && bAccepted == false; ++i)
+	// See if we need a new capture
+	if (ActiveLeftHoverCapture == nullptr)
 	{
-		FInputCaptureUpdate Result =
-			CaptureRequests[i].Source->BeginHoverCapture(Input, EInputCaptureSide::Left);
-		if (Result.State == EInputCaptureState::Begin)
+		for (int i = 0; i < CaptureRequests.Num(); ++i)
 		{
-			ActiveLeftHoverCapture = Result.Source;
-			ActiveLeftHoverCaptureOwner = CaptureRequests[i].Owner;
-			bAccepted = true;
-			return true;
+			FInputCaptureUpdate Result =
+				CaptureRequests[i].Source->BeginHoverCapture(Input, EInputCaptureSide::Left);
+			if (Result.State == EInputCaptureState::Begin)
+			{
+				ActiveLeftHoverCapture = Result.Source;
+				ActiveLeftHoverCaptureOwner = CaptureRequests[i].Owner;
+
+				// We say that the hover state has been modified, despite the fact that it's theoretically possible
+				// to end up with the same ActiveLeftHoverCapture if behaviors do some unpleasant things like claim
+				// that they want capture and then refuse it on BeginHoverCapture, or terminate capture in an update
+				// but then accept it on BeginHoverCapture... 
+				return true;
+			}
 		}
 	}
 	
-	// no hover! but we might have terminated an active hover
-	return bHoverStateModified;
+	return PreviousCapture != ActiveLeftHoverCapture;
 }
 
 
diff --git a/Engine/Source/Runtime/InteractiveToolsFramework/Public/InputRouter.h b/Engine/Source/Runtime/InteractiveToolsFramework/Public/InputRouter.h
--- a/Engine/Source/Runtime/InteractiveToolsFramework/Public/InputRouter.h
+++ b/Engine/Source/Runtime/InteractiveToolsFramework/Public/InputRouter.h
@@ -122,4 +122,5 @@ protected:
 
 	INTERACTIVETOOLSFRAMEWORK_API void TerminateHover(EInputCaptureSide Side);
 	INTERACTIVETOOLSFRAMEWORK_API bool ProcessMouseHover(const FInputDeviceState& Input);
+	INTERACTIVETOOLSFRAMEWORK_API bool UpdateExistingHoverCaptureIfPresent(const FInputDeviceState& Input);
 };
