fix issue where pressing modifier keys during a mouse drag could cause the drag to report spurious mouse input at 0,0 for a few ITF behaviors

diff --git a/Engine/Source/Runtime/InteractiveToolsFramework/Private/BaseBehaviors/MultiClickSequenceInputBehavior.cpp b/Engine/Source/Runtime/InteractiveToolsFramework/Private/BaseBehaviors/MultiClickSequenceInputBehavior.cpp
--- a/Engine/Source/Runtime/InteractiveToolsFramework/Private/BaseBehaviors/MultiClickSequenceInputBehavior.cpp
+++ b/Engine/Source/Runtime/InteractiveToolsFramework/Private/BaseBehaviors/MultiClickSequenceInputBehavior.cpp
@@ -63,6 +63,13 @@ FInputCaptureUpdate UMultiClickSequenceInputBehavior::UpdateCapture(const FInput
 
 	Modifiers.UpdateModifiers(input, Target);
 
+	// We have to check the device before going further because we get passed captures from 
+	// keyboard for modifier key press/releases, and those don't have valid mouse data.
+	if (!input.IsFromDevice(GetSupportedDevices()))
+	{
+		return FInputCaptureUpdate::Continue();
+	}
+
 	// Allow target to abort click sequence
 	if (State == ESequenceState::WaitingForNextClick && Target->RequestAbortClickSequence())
 	{
@@ -136,6 +143,13 @@ FInputCaptureUpdate UMultiClickSequenceInputBehavior::UpdateHoverCapture(const F
 
 	Modifiers.UpdateModifiers(InputState, Target);
 
+	// We have to check the device before going further because we get passed captures from 
+	// keyboard for modifier key press/releases, and those don't have valid mouse data.
+	if (!InputState.IsFromDevice(GetSupportedDevices()))
+	{
+		return FInputCaptureUpdate::Continue();
+	}
+
 	// Allow target to abort click sequence
 	if (State == ESequenceState::WaitingForNextClick && Target->RequestAbortClickSequence())
 	{
diff --git a/Engine/Source/Runtime/InteractiveToolsFramework/Private/BaseBehaviors/SingleClickOrDragBehavior.cpp b/Engine/Source/Runtime/InteractiveToolsFramework/Private/BaseBehaviors/SingleClickOrDragBehavior.cpp
--- a/Engine/Source/Runtime/InteractiveToolsFramework/Private/BaseBehaviors/SingleClickOrDragBehavior.cpp
+++ b/Engine/Source/Runtime/InteractiveToolsFramework/Private/BaseBehaviors/SingleClickOrDragBehavior.cpp
@@ -85,6 +85,13 @@ FInputCaptureUpdate USingleClickOrDragInputBehavior::UpdateCapture(const FInputD
 		Modifiers.UpdateModifiers(Input, DragTarget);
 	}
 
+	// We have to check the device before going further because we get passed captures from 
+	// keyboard for modifier key press/releases, and those don't have valid mouse data.
+	if (!Input.IsFromDevice(GetSupportedDevices()))
+	{
+		return FInputCaptureUpdate::Continue();
+	}
+
 	// check if mouse has moved far enough that we want to swap to drag behavior
 	if (bInDrag == false)
 	{
diff --git a/Engine/Source/Runtime/InteractiveToolsFramework/Private/BaseGizmos/ScalableSphereGizmo.cpp b/Engine/Source/Runtime/InteractiveToolsFramework/Private/BaseGizmos/ScalableSphereGizmo.cpp
--- a/Engine/Source/Runtime/InteractiveToolsFramework/Private/BaseGizmos/ScalableSphereGizmo.cpp
+++ b/Engine/Source/Runtime/InteractiveToolsFramework/Private/BaseGizmos/ScalableSphereGizmo.cpp
@@ -332,6 +332,13 @@ FInputCaptureUpdate UScalableSphereGizmoInputBehavior::BeginCapture(const FInput
 
 FInputCaptureUpdate UScalableSphereGizmoInputBehavior::UpdateCapture(const FInputDeviceState& input, const FInputCaptureData& data)
 {
+	// We have to check the device before going further because we get passed captures from 
+	// keyboard for modifier key press/releases, and those don't have valid ray data.
+	if (!input.IsFromDevice(GetSupportedDevices()))
+	{
+		return FInputCaptureUpdate::Continue();
+	}
+
 	FInputDeviceRay DeviceRay(input.Mouse.WorldRay, input.Mouse.Position2D);
 	LastWorldRay = DeviceRay.WorldRay;
 	LastScreenPosition = DeviceRay.ScreenPosition;
