New TRS Gizmo: fixed UGizmoElementBase::GetViewDependentVisibility that was using the local axis for angle computation
    - this had the effect of not hidding the axis and plane handles in the correct scenarios

#rnx

diff --git a/Engine/Source/Runtime/InteractiveToolsFramework/Private/BaseGizmos/GizmoElementBase.cpp b/Engine/Source/Runtime/InteractiveToolsFramework/Private/BaseGizmos/GizmoElementBase.cpp
--- a/Engine/Source/Runtime/InteractiveToolsFramework/Private/BaseGizmos/GizmoElementBase.cpp
+++ b/Engine/Source/Runtime/InteractiveToolsFramework/Private/BaseGizmos/GizmoElementBase.cpp
@@ -65,7 +65,7 @@ bool UGizmoElementBase::GetViewDependentVisibility(const FVector& InViewLocation
 	FVector ViewDir;
 	if (bInPerspectiveView)
 	{
-		FVector WorldCenter = InLocalToWorldTransform.TransformPosition(InLocalCenter);
+		const FVector WorldCenter = InLocalToWorldTransform.TransformPosition(InLocalCenter);
 		ViewDir = WorldCenter - InViewLocation;
 	}
 	else
@@ -74,18 +74,22 @@ bool UGizmoElementBase::GetViewDependentVisibility(const FVector& InViewLocation
 	}
 	ViewDir.Normalize();
 
-	bool bVisibleViewDependent;
-	if (ViewDependentType == EGizmoElementViewDependentType::Axis)
+	const FVector WorldViewDependentAxis = InLocalToWorldTransform.TransformVectorNoScale(ViewDependentAxis);
+	switch (ViewDependentType)
 	{
-		bVisibleViewDependent = FMath::Abs(FVector::DotProduct(ViewDependentAxis, ViewDir)) < DefaultViewDependentAxialMaxCosAngleTol;
-	}
-	else // (ViewDependentType == EGizmoElementViewDependentType::Plane)
-	{
-		bVisibleViewDependent = FMath::Abs(FVector::DotProduct(ViewDependentAxis, ViewDir)) > DefaultViewDependentPlanarMinCosAngleTol;
+		case EGizmoElementViewDependentType::None:
+			return true;
+		case EGizmoElementViewDependentType::Axis:
+			return FMath::Abs(FVector::DotProduct(WorldViewDependentAxis, ViewDir)) < DefaultViewDependentAxialMaxCosAngleTol;
+		case EGizmoElementViewDependentType::Plane:	
+			return FMath::Abs(FVector::DotProduct(WorldViewDependentAxis, ViewDir)) > DefaultViewDependentPlanarMinCosAngleTol;
+		default:
+			// check that we didn't had any new type
+			ensure(false);
+			break;
 	}
 
-	return bVisibleViewDependent;
-
+	return true;
 }
 
 bool UGizmoElementBase::GetViewAlignRot(const FVector& InViewLocation, const FVector& InViewDirection, const FVector& InViewUp, bool bInPerspectiveView, const FTransform& InLocalToWorldTransform, const FVector& InLocalCenter, FQuat& OutAlignRot) const
