InteractiveToolsFramework: Fix context object store FindContextByClass not being able to look up context objects based on UInterfaces.

diff --git a/Engine/Source/Runtime/InteractiveToolsFramework/Private/ContextObjectStore.cpp b/Engine/Source/Runtime/InteractiveToolsFramework/Private/ContextObjectStore.cpp
--- a/Engine/Source/Runtime/InteractiveToolsFramework/Private/ContextObjectStore.cpp
+++ b/Engine/Source/Runtime/InteractiveToolsFramework/Private/ContextObjectStore.cpp
@@ -6,11 +6,24 @@
 
 UObject* UContextObjectStore::FindContextByClass(UClass* InClass) const
 {
-	for (UObject* ContextObject : ContextObjects)
+	if (InClass->HasAnyClassFlags(CLASS_Interface))
 	{
-		if (ContextObject && ContextObject->IsA(InClass))
+		for (UObject* ContextObject : ContextObjects)
 		{
-			return ContextObject;
+			if (ContextObject && ContextObject->GetClass()->ImplementsInterface(InClass))
+			{
+				return ContextObject;
+			}
+		}
+	}
+	else
+	{
+		for (UObject* ContextObject : ContextObjects)
+		{
+			if (ContextObject && ContextObject->IsA(InClass))
+			{
+				return ContextObject;
+			}
 		}
 	}
 	
