Fixing build break from 30978118

diff --git a/Engine/Source/Runtime/InteractiveToolsFramework/Private/BaseTools/BaseBrushTool.cpp b/Engine/Source/Runtime/InteractiveToolsFramework/Private/BaseTools/BaseBrushTool.cpp
--- a/Engine/Source/Runtime/InteractiveToolsFramework/Private/BaseTools/BaseBrushTool.cpp
+++ b/Engine/Source/Runtime/InteractiveToolsFramework/Private/BaseTools/BaseBrushTool.cpp
@@ -105,9 +105,11 @@ void UBrushAdjusterInputBehavior::OnDragUpdate(FVector2D InScreenPosition)
 		NewRadius = FMath::Max(NewRadius, 0.01f);
 		BrushTool->BrushProperties->BrushRadius = NewRadius;
 		BrushTool->LastBrushStamp.Radius = NewRadius;
-		
+
+#if WITH_EDITOR
 		FPropertyChangedEvent PropertyChangedEvent(UBrushBaseProperties::StaticClass()->FindPropertyByName(GET_MEMBER_NAME_CHECKED(UBrushBaseProperties, BrushRadius)));
 		BrushTool->BrushProperties->PostEditChangeProperty(PropertyChangedEvent);
+#endif
 	}
 	else
 	{
@@ -115,9 +117,11 @@ void UBrushAdjusterInputBehavior::OnDragUpdate(FVector2D InScreenPosition)
 		float NewStrength = StartBrushStrength + VerticalDelta * -(StrengthAdjustSpeed * DPIScale);
 		NewStrength = FMath::Min(1.0f,FMath::Max(NewStrength, 0.f));
 		BrushTool->BrushProperties->BrushStrength = NewStrength;
-		
+
+#if WITH_EDITOR
 		FPropertyChangedEvent PropertyChangedEvent(UBrushBaseProperties::StaticClass()->FindPropertyByName(GET_MEMBER_NAME_CHECKED(UBrushBaseProperties, BrushStrength)));
 		BrushTool->BrushProperties->PostEditChangeProperty(PropertyChangedEvent);
+#endif
 	}	
 }
 
