Various fixes to skin paint tool
- improved brush falloff and strength rendering
- falloff updates in realtime in viewport when adjusted
- added multi-color bone support

diff --git a/Engine/Plugins/Experimental/MeshModelingToolsetExp/Source/MeshModelingToolsEditorOnlyExp/Private/SkeletalMesh/SkinWeightsPaintTool.cpp b/Engine/Plugins/Experimental/MeshModelingToolsetExp/Source/MeshModelingToolsEditorOnlyExp/Private/SkeletalMesh/SkinWeightsPaintTool.cpp
--- a/Engine/Plugins/Experimental/MeshModelingToolsetExp/Source/MeshModelingToolsEditorOnlyExp/Private/SkeletalMesh/SkinWeightsPaintTool.cpp
+++ b/Engine/Plugins/Experimental/MeshModelingToolsetExp/Source/MeshModelingToolsEditorOnlyExp/Private/SkeletalMesh/SkinWeightsPaintTool.cpp
@@ -732,6 +732,11 @@ void USkinWeightsPaintTool::Setup()
 	ReplaceToolPropertySource(BrushProperties, WeightToolProperties);
 	BrushProperties = WeightToolProperties;
 
+	// brush render customization
+	BrushStampIndicator->bScaleNormalByStrength = true;
+	BrushStampIndicator->SecondaryLineThickness = 1.0f;
+	BrushStampIndicator->SecondaryLineColor = FLinearColor::Yellow;
+
 	// default to the root bone as current bone
 	PendingCurrentBone = CurrentBone = Component->GetSkeletalMeshAsset()->GetRefSkeleton().GetBoneName(0);
 
diff --git a/Engine/Plugins/Experimental/MeshModelingToolsetExp/Source/MeshModelingToolsExp/Private/DynamicMeshSculptTool.cpp b/Engine/Plugins/Experimental/MeshModelingToolsetExp/Source/MeshModelingToolsExp/Private/DynamicMeshSculptTool.cpp
--- a/Engine/Plugins/Experimental/MeshModelingToolsetExp/Source/MeshModelingToolsExp/Private/DynamicMeshSculptTool.cpp
+++ b/Engine/Plugins/Experimental/MeshModelingToolsetExp/Source/MeshModelingToolsExp/Private/DynamicMeshSculptTool.cpp
@@ -1511,8 +1511,15 @@ void UDynamicMeshSculptTool::Render(IToolsContextRenderAPI* RenderAPI)
 
 	FViewCameraState RenderCameraState = RenderAPI->GetCameraState();
 
-	//BrushIndicator->Update( (float)this->CurrentBrushRadius, (FVector)this->LastBrushPosWorld, (FVector)this->LastBrushPosNormalWorld, 1.0f-BrushProperties->BrushFalloffAmount);
-	BrushIndicator->Update( (float)this->CurrentBrushRadius, (FVector)this->LastBrushPosWorld, (FVector)this->LastBrushPosNormalWorld, 0.5f);
+	constexpr float Falloff = 0.5f;
+	constexpr float Strength = 1.0f;
+	BrushIndicator->Update(
+		(float)this->CurrentBrushRadius,
+		(FVector)this->LastBrushPosWorld,
+		(FVector)this->LastBrushPosNormalWorld,
+		Falloff,
+		Strength);
+	
 	if (BrushIndicatorMaterial)
 	{
 		double FixedDimScale = ToolSceneQueriesUtil::CalculateDimensionFromVisualAngleD(RenderCameraState, LastBrushPosWorld, 1.5f);
diff --git a/Engine/Source/Runtime/InteractiveToolsFramework/Private/BaseGizmos/BrushStampIndicator.cpp b/Engine/Source/Runtime/InteractiveToolsFramework/Private/BaseGizmos/BrushStampIndicator.cpp
--- a/Engine/Source/Runtime/InteractiveToolsFramework/Private/BaseGizmos/BrushStampIndicator.cpp
+++ b/Engine/Source/Runtime/InteractiveToolsFramework/Private/BaseGizmos/BrushStampIndicator.cpp
@@ -36,17 +36,23 @@ void UBrushStampIndicator::Render(IToolsContextRenderAPI* RenderAPI)
 		FToolDataVisualizer Draw;
 		Draw.BeginFrame(RenderAPI);
 
-		if (bDrawRadiusCircle)
+		if (bDrawSecondaryLines)
 		{
-			Draw.DrawCircle(BrushPosition, BrushNormal, BrushRadius, SampleStepCount, LineColor, LineThickness, bDepthTested);
+			const bool bCirclesCoincident = bDrawRadiusCircle && FMath::IsNearlyEqual(BrushFalloff, 1.0f);
+			if (!bCirclesCoincident)
+			{
+				Draw.DrawCircle(BrushPosition, BrushNormal, BrushRadius*BrushFalloff, SampleStepCount, SecondaryLineColor, SecondaryLineThickness, bDepthTested);
+			}
+			
+			const float NormalScale = bScaleNormalByStrength ? BrushStrength : 1.0f;
+			Draw.DrawLine(BrushPosition, BrushPosition + BrushRadius * BrushNormal * NormalScale, SecondaryLineColor, SecondaryLineThickness, bDepthTested);
 		}
-
-		if (bDrawSecondaryLines)
+		
+		if (bDrawRadiusCircle)
 		{
-			Draw.DrawCircle(BrushPosition, BrushNormal, BrushRadius*BrushFalloff, SampleStepCount, SecondaryLineColor, SecondaryLineThickness, bDepthTested);
-			Draw.DrawLine(BrushPosition, BrushPosition + BrushRadius * BrushNormal, SecondaryLineColor, SecondaryLineThickness, bDepthTested);
+			Draw.DrawCircle(BrushPosition, BrushNormal, BrushRadius, SampleStepCount, LineColor, LineThickness, bDepthTested);
 		}
-
+		
 		Draw.EndFrame();
 	}
 }
@@ -56,12 +62,13 @@ void UBrushStampIndicator::Tick(float DeltaTime)
 }
 
 
-void UBrushStampIndicator::Update(float Radius, const FVector& Position, const FVector& Normal, float Falloff)
+void UBrushStampIndicator::Update(float Radius, const FVector& Position, const FVector& Normal, float Falloff, float Strength)
 {
 	BrushRadius = Radius;
 	BrushPosition = Position;
 	BrushNormal = Normal;
 	BrushFalloff = Falloff;
+	BrushStrength = Strength;
 
 	if (AttachedComponent != nullptr)
 	{
diff --git a/Engine/Source/Runtime/InteractiveToolsFramework/Private/BaseTools/BaseBrushTool.cpp b/Engine/Source/Runtime/InteractiveToolsFramework/Private/BaseTools/BaseBrushTool.cpp
--- a/Engine/Source/Runtime/InteractiveToolsFramework/Private/BaseTools/BaseBrushTool.cpp
+++ b/Engine/Source/Runtime/InteractiveToolsFramework/Private/BaseTools/BaseBrushTool.cpp
@@ -514,7 +514,7 @@ void UBaseBrushTool::UpdateBrushStampIndicator()
 			BrushStampIndicator->LineColor = BrushAdjusterBehavior->IsBrushBeingAdjusted() ? FLinearColor::White : FLinearColor::Green;	
 		}
 		
-		BrushStampIndicator->Update(BrushProperties->BrushRadius, LastBrushStamp.WorldPosition, LastBrushStamp.WorldNormal, LastBrushStamp.Falloff);
+		BrushStampIndicator->Update(BrushProperties->BrushRadius, LastBrushStamp.WorldPosition, LastBrushStamp.WorldNormal, BrushProperties->BrushFalloffAmount, BrushProperties->BrushStrength);
 	}
 }
 
diff --git a/Engine/Source/Runtime/InteractiveToolsFramework/Public/BaseGizmos/BrushStampIndicator.h b/Engine/Source/Runtime/InteractiveToolsFramework/Public/BaseGizmos/BrushStampIndicator.h
--- a/Engine/Source/Runtime/InteractiveToolsFramework/Public/BaseGizmos/BrushStampIndicator.h
+++ b/Engine/Source/Runtime/InteractiveToolsFramework/Public/BaseGizmos/BrushStampIndicator.h
@@ -40,7 +40,7 @@ public:
 	/**
 	 * Update the Radius, Position, and Normal of the stamp indicator
 	 */
-	INTERACTIVETOOLSFRAMEWORK_API virtual void Update(float Radius, const FVector& Position, const FVector& Normal, float Falloff);
+	INTERACTIVETOOLSFRAMEWORK_API virtual void Update(float Radius, const FVector& Position, const FVector& Normal, float Falloff, float Strength);
 
 	/**
 	* Update the Radius, Transform and Falloff of the Stamp Indicator
@@ -59,12 +59,15 @@ public:
 	UPROPERTY()
 	float BrushFalloff = 0.5f;
 
+	UPROPERTY()
+	float BrushStrength = 1.0f;
+
 	UPROPERTY()
 	FVector BrushPosition = FVector::ZeroVector;
 
 	UPROPERTY()
 	FVector BrushNormal = FVector(0, 0, 1);;
-
+	
 
 
 	UPROPERTY()
@@ -73,6 +76,9 @@ public:
 	UPROPERTY()
 	bool bDrawRadiusCircle = true;
 
+	UPROPERTY()
+	bool bScaleNormalByStrength = false;
+
 	UPROPERTY()
 	int SampleStepCount = 32;
 
