Add DrawDisc to ToolDataVisualizer

The existing DrawCircle function draws a circle that cannot be filled

diff --git a/Engine/Source/Runtime/InteractiveToolsFramework/Private/ToolDataVisualizer.cpp b/Engine/Source/Runtime/InteractiveToolsFramework/Private/ToolDataVisualizer.cpp
--- a/Engine/Source/Runtime/InteractiveToolsFramework/Private/ToolDataVisualizer.cpp
+++ b/Engine/Source/Runtime/InteractiveToolsFramework/Private/ToolDataVisualizer.cpp
@@ -1,7 +1,7 @@
 // Copyright Epic Games, Inc. All Rights Reserved.
 
 #include "ToolDataVisualizer.h"
-#include "SceneManagement.h" 
+#include "MaterialShared.h"
 #include "ToolContextInterfaces.h"
 #include "SceneManagement.h"   // DrawCircle
 #include "BaseGizmos/GizmoMath.h"
@@ -212,4 +212,20 @@ void FToolDataVisualizer::InternalDrawViewFacingX(const FVector& Position, float
 		WorldPosition + UpOffset - RightOffset,
 		WorldPosition - UpOffset + RightOffset,
 		Color, LineThicknessIn, bDepthTestedIn);
+}
+
+void FToolDataVisualizer::InternalDrawDisc(const FVector& Position, const FVector& Normal, float Radius, int Steps, const FColor& Color, FMaterialRenderProxy* RenderProxy, bool bDepthTestedIn)
+{
+	const FVector TransformNormal = TransformN(Normal);
+	FVector Tan1;
+	FVector Tan2;
+	GizmoMath::MakeNormalPlaneBasis(TransformNormal, Tan1, Tan2);
+	Tan1.Normalize(); 
+	Tan2.Normalize();
+
+	// this function is from SceneManagement.h
+	const uint8 DepthPriority = (bDepthTestedIn) ? SDPG_World : SDPG_Foreground;
+	const FVector TransformPosition = TransformP(Position);
+	::DrawDisc(CurrentPDI, TransformPosition, Tan1, Tan2, Color, Radius, Steps, RenderProxy,
+	           DepthPriority);
 }
\ No newline at end of file
diff --git a/Engine/Source/Runtime/InteractiveToolsFramework/Public/ToolDataVisualizer.h b/Engine/Source/Runtime/InteractiveToolsFramework/Public/ToolDataVisualizer.h
--- a/Engine/Source/Runtime/InteractiveToolsFramework/Public/ToolDataVisualizer.h
+++ b/Engine/Source/Runtime/InteractiveToolsFramework/Public/ToolDataVisualizer.h
@@ -7,6 +7,7 @@
 
 class FPrimitiveDrawInterface;
 class IToolsContextRenderAPI;
+class FMaterialRenderProxy;
 
 /**
  * FToolDataVisualizer is a utility class for Tool and Gizmo implementations
@@ -229,6 +230,12 @@ public:
 		InternalDrawSquare(Center, SideA, SideB, LineColor, LineThickness, bDepthTested);
 	}
 
+	/** Draw a 3D disc at given position/normal with the given parameters. Tangent axes are defined internally */
+	template<typename PointType>
+	void DrawDisc(const PointType& Position, const PointType& Normal, float Radius, int Steps, const FColor& Color, FMaterialRenderProxy* RenderProxy, bool bDepthTestedIn)
+	{
+		InternalDrawDisc(Position, Normal, Radius, Steps, Color, RenderProxy, bDepthTestedIn);
+	}
 
 protected:
 	/** We use this for drawing, extracted in InitializeFrame */
@@ -263,5 +270,6 @@ protected:
 	INTERACTIVETOOLSFRAMEWORK_API virtual void InternalDrawSquare(const FVector& Center, const FVector& SideA, const FVector& SideB, const FLinearColor& Color, float LineThickness, bool bDepthTested);
 	// actually does the X drawing
 	INTERACTIVETOOLSFRAMEWORK_API virtual void InternalDrawViewFacingX(const FVector& Position, float Width, const FLinearColor& Color, float LineThickness, bool bDepthTested);
-
+	// actually does the disc drawing
+	INTERACTIVETOOLSFRAMEWORK_API virtual void InternalDrawDisc(const FVector& Position, const FVector& Normal, float Radius, int Steps, const FColor& Color, FMaterialRenderProxy* RenderProxy, bool bDepthTested);
 };
