ToolDataVisualizer - fix type cast warning.

diff --git a/Engine/Source/Runtime/InteractiveToolsFramework/Private/ToolDataVisualizer.cpp b/Engine/Source/Runtime/InteractiveToolsFramework/Private/ToolDataVisualizer.cpp
--- a/Engine/Source/Runtime/InteractiveToolsFramework/Private/ToolDataVisualizer.cpp
+++ b/Engine/Source/Runtime/InteractiveToolsFramework/Private/ToolDataVisualizer.cpp
@@ -224,7 +224,7 @@ void FToolDataVisualizer::InternalDrawDisc(const FVector& Position, const FVecto
 	Tan2.Normalize();
 
 	// this function is from SceneManagement.h
-	const uint8 DepthPriority = (bDepthTestedIn) ? SDPG_World : SDPG_Foreground;
+	const uint8 DepthPriority = uint8( (bDepthTestedIn) ? SDPG_World : SDPG_Foreground );
 	const FVector TransformPosition = TransformP(Position);
 	::DrawDisc(CurrentPDI, TransformPosition, Tan1, Tan2, Color, Radius, Steps, RenderProxy,
 	           DepthPriority);
