Speed when adjusting brush size is now proportional to distance from camera.

diff --git a/Engine/Source/Runtime/InteractiveToolsFramework/Private/BaseTools/BaseBrushTool.cpp b/Engine/Source/Runtime/InteractiveToolsFramework/Private/BaseTools/BaseBrushTool.cpp
--- a/Engine/Source/Runtime/InteractiveToolsFramework/Private/BaseTools/BaseBrushTool.cpp
+++ b/Engine/Source/Runtime/InteractiveToolsFramework/Private/BaseTools/BaseBrushTool.cpp
@@ -101,7 +101,7 @@ void UBrushAdjusterInputBehavior::OnDragUpdate(FVector2D InScreenPosition)
 	if (bAdjustingHorizontally)
 	{
 		// adjust brush size based on horizontal mouse drag
-		float NewRadius = StartBrushRadius + HorizontalDelta * (SizeAdjustSpeed * DPIScale);
+		float NewRadius = StartBrushRadius + HorizontalDelta * (SizeAdjustSpeed * DPIScale * BrushTool->LastBrushStamp.HitResult.Distance);
 		NewRadius = FMath::Max(NewRadius, 0.01f);
 		BrushTool->BrushProperties->BrushRadius = NewRadius;
 		BrushTool->LastBrushStamp.Radius = NewRadius;
diff --git a/Engine/Source/Runtime/InteractiveToolsFramework/Public/BaseTools/BaseBrushTool.h b/Engine/Source/Runtime/InteractiveToolsFramework/Public/BaseTools/BaseBrushTool.h
--- a/Engine/Source/Runtime/InteractiveToolsFramework/Public/BaseTools/BaseBrushTool.h
+++ b/Engine/Source/Runtime/InteractiveToolsFramework/Public/BaseTools/BaseBrushTool.h
@@ -128,9 +128,9 @@ private:
 	bool bAdjustingHorizontally = true;
 
 	// the speed (centimeters per unit of screen coordinate) to adjust brush size when dragging
-	float SizeAdjustSpeed = 0.04f; // sensible default based on 1080p monitor
+	float SizeAdjustSpeed = 0.002f; // sensible default
 	// the speed (in strength per unit of screen coordinate) to adjust brush strength when dragging
-	float StrengthAdjustSpeed = 0.005f; // sensible default based on 1080p monitor
+	float StrengthAdjustSpeed = 0.005f; // sensible default
 
 	// the target brush tool to adjust
 	UBaseBrushTool* BrushTool;
