ModelingTools: Allow BrushAdjusterInputBehavior to drive BrushSize when bSpecifyRadius is false.

diff --git a/Engine/Source/Runtime/InteractiveToolsFramework/Private/BaseTools/BaseBrushTool.cpp b/Engine/Source/Runtime/InteractiveToolsFramework/Private/BaseTools/BaseBrushTool.cpp
--- a/Engine/Source/Runtime/InteractiveToolsFramework/Private/BaseTools/BaseBrushTool.cpp
+++ b/Engine/Source/Runtime/InteractiveToolsFramework/Private/BaseTools/BaseBrushTool.cpp
@@ -65,7 +65,14 @@ void UBrushAdjusterInputBehavior::ResetAdjustmentOrigin(FVector2D InScreenPositi
 {
 	bAdjustingHorizontally = bHorizontalAdjust;
 	AdjustmentOrigin = InScreenPosition;
-	StartBrushRadius = BrushTool->BrushProperties->BrushRadius;
+	if (BrushTool->BrushProperties->bSpecifyRadius)
+	{
+		StartBrushRadius = BrushTool->BrushProperties->BrushRadius;
+	}
+	else
+	{
+		StartBrushRadius = BrushTool->BrushProperties->BrushSize;
+	}
 	StartBrushStrength = BrushTool->BrushProperties->BrushStrength;
 }
 
@@ -103,13 +110,22 @@ void UBrushAdjusterInputBehavior::OnDragUpdate(FVector2D InScreenPosition)
 		// adjust brush size based on horizontal mouse drag
 		float NewRadius = StartBrushRadius + HorizontalDelta * (SizeAdjustSpeed * DPIScale * BrushTool->LastBrushStamp.HitResult.Distance);
 		NewRadius = FMath::Max(NewRadius, 0.01f);
-		BrushTool->BrushProperties->BrushRadius = NewRadius;
-		BrushTool->LastBrushStamp.Radius = NewRadius;
-
+		if (BrushTool->BrushProperties->bSpecifyRadius)
+		{
+			BrushTool->BrushProperties->BrushRadius = NewRadius;
 #if WITH_EDITOR
-		FPropertyChangedEvent PropertyChangedEvent(UBrushBaseProperties::StaticClass()->FindPropertyByName(GET_MEMBER_NAME_CHECKED(UBrushBaseProperties, BrushRadius)));
-		BrushTool->BrushProperties->PostEditChangeProperty(PropertyChangedEvent);
+			FPropertyChangedEvent PropertyChangedEvent(UBrushBaseProperties::StaticClass()->FindPropertyByName(GET_MEMBER_NAME_CHECKED(UBrushBaseProperties, BrushRadius)));
+			BrushTool->BrushProperties->PostEditChangeProperty(PropertyChangedEvent);
+#endif
+		}
+		else
+		{
+			BrushTool->BrushProperties->BrushSize = NewRadius;
+#if WITH_EDITOR
+			FPropertyChangedEvent PropertyChangedEvent(UBrushBaseProperties::StaticClass()->FindPropertyByName(GET_MEMBER_NAME_CHECKED(UBrushBaseProperties, BrushSize)));
+			BrushTool->BrushProperties->PostEditChangeProperty(PropertyChangedEvent);
 #endif
+		}
 	}
 	else
 	{
