InteractiveToolsFramework: Separate the BaseBrushTool's BrushAdjusterInputBehavior click drag handling into its own behavior.

Previously BrushAdjusterInputBehavior was used to set a toggle that would be read during the On*Drag callbacks to differentiate between a "brush stroke" and a "brush adjustment". This was error prone for both new and old tools that derived off of BaseBrushTool as it broke or was easy to break assumptions made in the tool's drag callbacks. With this separation, OnBeginDrag/OnUpdateDrag/OnEndDrag will only be called for brush strokes.

Additional minor fixes:
- Disabled the BrushAdjusterBehavior for SeamSculpt.
- Fixed initial BrushAdjuster HUD display prior to first drag.

diff --git a/Engine/Plugins/Experimental/MeshModelingToolsetExp/Source/MeshModelingToolsEditorOnlyExp/Private/SkeletalMesh/SkinWeightsPaintTool.cpp b/Engine/Plugins/Experimental/MeshModelingToolsetExp/Source/MeshModelingToolsEditorOnlyExp/Private/SkeletalMesh/SkinWeightsPaintTool.cpp
--- a/Engine/Plugins/Experimental/MeshModelingToolsetExp/Source/MeshModelingToolsEditorOnlyExp/Private/SkeletalMesh/SkinWeightsPaintTool.cpp
+++ b/Engine/Plugins/Experimental/MeshModelingToolsetExp/Source/MeshModelingToolsEditorOnlyExp/Private/SkeletalMesh/SkinWeightsPaintTool.cpp
@@ -1240,25 +1240,20 @@ void USkinWeightsPaintTool::OnBeginDrag(const FRay& WorldRay)
 {
 	UDynamicMeshBrushTool::OnBeginDrag(WorldRay);
 
-	if (IsInBrushStroke())
-	{
-		bInvertStroke = GetCtrlToggle();
-		BeginChange();
-		StartStamp = UBaseBrushTool::LastBrushStamp;
-		LastStamp = StartStamp;
-		bStampPending = true;
-		LongTransactions.Open(LOCTEXT("PaintWeightChange", "Paint skin weights."), GetToolManager());
-	}
+	bInvertStroke = GetCtrlToggle();
+	BeginChange();
+	StartStamp = UBaseBrushTool::LastBrushStamp;
+	LastStamp = StartStamp;
+	bStampPending = true;
+	LongTransactions.Open(LOCTEXT("PaintWeightChange", "Paint skin weights."), GetToolManager());
 }
 
 void USkinWeightsPaintTool::OnUpdateDrag(const FRay& WorldRay)
 {
 	UDynamicMeshBrushTool::OnUpdateDrag(WorldRay);
-	if (IsInBrushStroke())
-	{
-		LastStamp = UBaseBrushTool::LastBrushStamp;
-		bStampPending = true;
-	}
+	
+	LastStamp = UBaseBrushTool::LastBrushStamp;
+	bStampPending = true;
 }
 
 void USkinWeightsPaintTool::OnEndDrag(const FRay& Ray)
diff --git a/Engine/Plugins/Experimental/MeshModelingToolsetExp/Source/MeshModelingToolsExp/Private/MeshAttributePaintTool.cpp b/Engine/Plugins/Experimental/MeshModelingToolsetExp/Source/MeshModelingToolsExp/Private/MeshAttributePaintTool.cpp
--- a/Engine/Plugins/Experimental/MeshModelingToolsetExp/Source/MeshModelingToolsExp/Private/MeshAttributePaintTool.cpp
+++ b/Engine/Plugins/Experimental/MeshModelingToolsetExp/Source/MeshModelingToolsExp/Private/MeshAttributePaintTool.cpp
@@ -328,15 +328,12 @@ void UMeshAttributePaintTool::OnBeginDrag(const FRay& WorldRay)
 	UDynamicMeshBrushTool::OnBeginDrag(WorldRay);
 
 	PreviewBrushROI.Reset();
-	if (IsInBrushStroke())
-	{
-		bInRemoveStroke = GetCtrlToggle();
-		bInSmoothStroke = GetShiftToggle();
-		BeginChange();
-		StartStamp = UBaseBrushTool::LastBrushStamp;
-		LastStamp = StartStamp;
-		bStampPending = true;
-	}
+	bInRemoveStroke = GetCtrlToggle();
+	bInSmoothStroke = GetShiftToggle();
+	BeginChange();
+	StartStamp = UBaseBrushTool::LastBrushStamp;
+	LastStamp = StartStamp;
+	bStampPending = true;
 }
 
 
@@ -344,34 +341,26 @@ void UMeshAttributePaintTool::OnBeginDrag(const FRay& WorldRay)
 void UMeshAttributePaintTool::OnUpdateDrag(const FRay& WorldRay)
 {
 	UDynamicMeshBrushTool::OnUpdateDrag(WorldRay);
-	if (IsInBrushStroke())
-	{
-		LastStamp = UBaseBrushTool::LastBrushStamp;
-		bStampPending = true;
-	}
+
+	LastStamp = UBaseBrushTool::LastBrushStamp;
+	bStampPending = true;
 }
 
 
 
 void UMeshAttributePaintTool::OnEndDrag(const FRay& Ray)
 {
-	// Capture brush stroke state prior to invoking Super::OnEndDrag
-	const bool bWasInBrushStroke = IsInBrushStroke();
-	
 	UDynamicMeshBrushTool::OnEndDrag(Ray);
 
 	bInRemoveStroke = bInSmoothStroke = false;
 	bStampPending = false;
 
-	if (bWasInBrushStroke)
+	// close change record
+	TUniquePtr<FMeshAttributePaintChange> Change = EndChange();
+	if (Change)
 	{
-		// close change record
-		TUniquePtr<FMeshAttributePaintChange> Change = EndChange();
-		if (Change)
-		{
-			GetToolManager()->EmitObjectChange(this, MoveTemp(Change), LOCTEXT("AttributeValuesChange", "Paint"));
-			LongTransactions.Close(GetToolManager());
-		}
+		GetToolManager()->EmitObjectChange(this, MoveTemp(Change), LOCTEXT("AttributeValuesChange", "Paint"));
+		LongTransactions.Close(GetToolManager());
 	}
 }
 
diff --git a/Engine/Plugins/Experimental/MeshModelingToolsetExp/Source/MeshModelingToolsExp/Private/MeshSelectionTool.cpp b/Engine/Plugins/Experimental/MeshModelingToolsetExp/Source/MeshModelingToolsExp/Private/MeshSelectionTool.cpp
--- a/Engine/Plugins/Experimental/MeshModelingToolsetExp/Source/MeshModelingToolsExp/Private/MeshSelectionTool.cpp
+++ b/Engine/Plugins/Experimental/MeshModelingToolsetExp/Source/MeshModelingToolsExp/Private/MeshSelectionTool.cpp
@@ -448,15 +448,12 @@ void UMeshSelectionTool::OnBeginDrag(const FRay& WorldRay)
 	UDynamicMeshBrushTool::OnBeginDrag(WorldRay);
 
 	PreviewBrushROI.Reset();
-	if (IsInBrushStroke())
-	{
-		bInRemoveStroke = GetShiftToggle();
-		BeginChange(bInRemoveStroke == false);
-		StartStamp = UBaseBrushTool::LastBrushStamp;
-		LastStamp = StartStamp;
-		bStampPending = true;
-		LongTransactions.Open(LOCTEXT("MeshSelectionChange", "Mesh Selection"), GetToolManager());
-	}
+	bInRemoveStroke = GetShiftToggle();
+	BeginChange(bInRemoveStroke == false);
+	StartStamp = UBaseBrushTool::LastBrushStamp;
+	LastStamp = StartStamp;
+	bStampPending = true;
+	LongTransactions.Open(LOCTEXT("MeshSelectionChange", "Mesh Selection"), GetToolManager());
 }
 
 
@@ -464,11 +461,8 @@ void UMeshSelectionTool::OnBeginDrag(const FRay& WorldRay)
 void UMeshSelectionTool::OnUpdateDrag(const FRay& WorldRay)
 {
 	UDynamicMeshBrushTool::OnUpdateDrag(WorldRay);
-	if (IsInBrushStroke())
-	{
-		LastStamp = UBaseBrushTool::LastBrushStamp;
-		bStampPending = true;
-	}
+	LastStamp = UBaseBrushTool::LastBrushStamp;
+	bStampPending = true;
 }
 
 
@@ -752,23 +746,17 @@ void UMeshSelectionTool::UpdateFaceSelection(const FBrushStampData& Stamp, const
 
 void UMeshSelectionTool::OnEndDrag(const FRay& Ray)
 {
-	// Capture brush stroke state prior to invoking Super::OnEndDrag
-	const bool bWasInBrushStroke = IsInBrushStroke();
-	
 	UDynamicMeshBrushTool::OnEndDrag(Ray);
 
 	bInRemoveStroke = false;
 	bStampPending = false;
 
 	// close change record
-	if (bWasInBrushStroke)
+	TUniquePtr<FToolCommandChange> Change = EndChange();
+	if (Change)
 	{
-		TUniquePtr<FToolCommandChange> Change = EndChange();
-		if (Change)
-		{
-			GetToolManager()->EmitObjectChange(Selection, MoveTemp(Change), LOCTEXT("MeshSelectionChange", "Mesh Selection"));
-			LongTransactions.Close(GetToolManager());
-		}
+		GetToolManager()->EmitObjectChange(Selection, MoveTemp(Change), LOCTEXT("MeshSelectionChange", "Mesh Selection"));
+		LongTransactions.Close(GetToolManager());
 	}
 }
 
diff --git a/Engine/Plugins/Experimental/MeshModelingToolsetExp/Source/MeshModelingToolsExp/Private/SeamSculptTool.cpp b/Engine/Plugins/Experimental/MeshModelingToolsetExp/Source/MeshModelingToolsExp/Private/SeamSculptTool.cpp
--- a/Engine/Plugins/Experimental/MeshModelingToolsetExp/Source/MeshModelingToolsExp/Private/SeamSculptTool.cpp
+++ b/Engine/Plugins/Experimental/MeshModelingToolsetExp/Source/MeshModelingToolsExp/Private/SeamSculptTool.cpp
@@ -151,7 +151,7 @@ void USeamSculptTool::OnBeginDrag(const FRay& Ray)
 {
 	UBaseBrushTool::OnBeginDrag(Ray);
 
-	if (IsInBrushStroke() && CurrentSnapVertex >= 0)
+	if (CurrentSnapVertex >= 0)
 	{
 		CaptureState = EActiveCaptureState::DrawNewPath;
 		DrawPathStartVertex = CurrentSnapVertex;
@@ -188,12 +188,9 @@ void USeamSculptTool::OnUpdateDrag(const FRay& Ray)
 
 void USeamSculptTool::OnEndDrag(const FRay& Ray)
 {
-	// Capture brush stroke state prior to invoking Super::OnEndDrag
-	const bool bWasInBrushStroke = IsInBrushStroke();
-	
 	UBaseBrushTool::OnEndDrag(Ray);
 
-	if (bWasInBrushStroke && CaptureState == EActiveCaptureState::DrawNewPath)
+	if (CaptureState == EActiveCaptureState::DrawNewPath)
 	{
 		CreateSeamAlongPath();
 		CurDrawPath.Reset();
diff --git a/Engine/Plugins/Experimental/MeshModelingToolsetExp/Source/MeshModelingToolsExp/Public/SeamSculptTool.h b/Engine/Plugins/Experimental/MeshModelingToolsetExp/Source/MeshModelingToolsExp/Public/SeamSculptTool.h
--- a/Engine/Plugins/Experimental/MeshModelingToolsetExp/Source/MeshModelingToolsExp/Public/SeamSculptTool.h
+++ b/Engine/Plugins/Experimental/MeshModelingToolsetExp/Source/MeshModelingToolsExp/Public/SeamSculptTool.h
@@ -69,6 +69,7 @@ public:
 	virtual void OnUpdateDrag(const FRay& Ray) override;
 	virtual void OnEndDrag(const FRay& Ray) override;
 	virtual bool OnUpdateHover(const FInputDeviceRay& DevicePos) override;
+	virtual bool SupportsBrushAdjustmentInput() override { return false; }
 
 	// IInteractiveToolManageGeometrySelectionAPI -- this tool won't update external geometry selection or change selection-relevant mesh IDs
 	virtual bool IsInputSelectionValidOnOutput() override
diff --git a/Engine/Source/Runtime/InteractiveToolsFramework/Private/BaseTools/BaseBrushTool.cpp b/Engine/Source/Runtime/InteractiveToolsFramework/Private/BaseTools/BaseBrushTool.cpp
--- a/Engine/Source/Runtime/InteractiveToolsFramework/Private/BaseTools/BaseBrushTool.cpp
+++ b/Engine/Source/Runtime/InteractiveToolsFramework/Private/BaseTools/BaseBrushTool.cpp
@@ -225,7 +225,37 @@ void UBaseBrushTool::Setup()
 	{
 		BrushAdjusterBehavior = NewObject<UBrushAdjusterInputBehavior>(this);
 		BrushAdjusterBehavior->Initialize(this);
-		AddInputBehavior(BrushAdjusterBehavior.Get());	
+		AddInputBehavior(BrushAdjusterBehavior.Get());
+
+		ULocalClickDragInputBehavior* BrushAdjusterClickDragBehavior = NewObject<ULocalClickDragInputBehavior>(this);
+		BrushAdjusterClickDragBehavior->Initialize();
+		BrushAdjusterClickDragBehavior->SetDefaultPriority(FInputCapturePriority().MakeHigher());
+		BrushAdjusterClickDragBehavior->ModifierCheckFunc = [this](const FInputDeviceState&)
+		{
+			return BrushAdjusterBehavior.IsValid() ? BrushAdjusterBehavior->IsBrushBeingAdjusted() : false;
+		};
+		BrushAdjusterClickDragBehavior->CanBeginClickDragFunc = [](const FInputDeviceRay&)
+		{
+			// fake screen hit
+			return FInputRayHit(0.f);
+		};
+		BrushAdjusterClickDragBehavior->OnClickPressFunc = [this](const FInputDeviceRay& PressPos)
+		{
+			if (BrushAdjusterBehavior.IsValid())
+			{
+				BrushAdjusterBehavior->OnDragStart(PressPos.ScreenPosition);
+			}
+		};
+		BrushAdjusterClickDragBehavior->OnClickDragFunc = [this](const FInputDeviceRay& DragPos)
+		{
+			if (BrushAdjusterBehavior.IsValid())
+			{
+				BrushAdjusterBehavior->OnDragUpdate(DragPos.ScreenPosition);
+				RecalculateBrushRadius();
+				NotifyOfPropertyChangeByTool(BrushProperties);
+			}
+		};
+		AddInputBehavior(BrushAdjusterClickDragBehavior);
 	}
 }
 
@@ -251,35 +281,11 @@ FInputRayHit UBaseBrushTool::CanBeginClickDragSequence(const FInputDeviceRay& Pr
 		// no hit
 		return FInputRayHit();
 	}
-	
-	if (BrushAdjusterBehavior.IsValid() && BrushAdjusterBehavior->IsBrushBeingAdjusted())
-	{
-		// fake screen hit
-		return FInputRayHit(0.f);
-	}
 
 	// hit-test the tool target
 	return Super::CanBeginClickDragSequence(PressPos);	
 }
 
-void UBaseBrushTool::OnClickPress(const FInputDeviceRay& PressPos)
-{
-	Super::OnClickPress(PressPos);
-	if (BrushAdjusterBehavior.IsValid())
-	{
-		BrushAdjusterBehavior->OnDragStart(PressPos.ScreenPosition);
-	}
-}
-
-void UBaseBrushTool::OnClickDrag(const FInputDeviceRay& DragPos)
-{
-	Super::OnClickDrag(DragPos);
-	if (BrushAdjusterBehavior.IsValid())
-	{
-		BrushAdjusterBehavior->OnDragUpdate(DragPos.ScreenPosition);
-	}
-}
-
 void UBaseBrushTool::IncreaseBrushSizeAction()
 {
 	if (BrushProperties->bSpecifyRadius)
@@ -424,12 +430,6 @@ void UBaseBrushTool::RecalculateBrushRadius()
 
 void UBaseBrushTool::OnBeginDrag(const FRay& Ray)
 {
-	if (BrushAdjusterBehavior.IsValid() && BrushAdjusterBehavior->IsBrushBeingAdjusted())
-	{
-		bInBrushStroke = false;
-		return;
-	}
-	
 	FHitResult OutHit;
 	if (HitTest(Ray, OutHit))
 	{
@@ -439,18 +439,14 @@ void UBaseBrushTool::OnBeginDrag(const FRay& Ray)
 		LastBrushStamp.HitResult = OutHit;
 		LastBrushStamp.Falloff = BrushProperties->BrushFalloffAmount;
 	}
+
+	PRAGMA_DISABLE_DEPRECATION_WARNINGS
 	bInBrushStroke = true;
+	PRAGMA_ENABLE_DEPRECATION_WARNINGS
 }
 
 void UBaseBrushTool::OnUpdateDrag(const FRay& Ray)
 {
-	if (BrushAdjusterBehavior.IsValid() && BrushAdjusterBehavior->IsBrushBeingAdjusted())
-	{
-		RecalculateBrushRadius();
-		NotifyOfPropertyChangeByTool(BrushProperties);
-		return;
-	}
-	
 	FHitResult OutHit;
 	if (HitTest(Ray, OutHit))
 	{
@@ -464,19 +460,30 @@ void UBaseBrushTool::OnUpdateDrag(const FRay& Ray)
 
 void UBaseBrushTool::OnEndDrag(const FRay& Ray)
 {
+	PRAGMA_DISABLE_DEPRECATION_WARNINGS
 	bInBrushStroke = false;
+	PRAGMA_ENABLE_DEPRECATION_WARNINGS
 }
 
 void UBaseBrushTool::OnCancelDrag()
 {
+	PRAGMA_DISABLE_DEPRECATION_WARNINGS
 	bInBrushStroke = false;
+	PRAGMA_ENABLE_DEPRECATION_WARNINGS
 }
 
 bool UBaseBrushTool::OnUpdateHover(const FInputDeviceRay& DevicePos)
 {
-	if (BrushAdjusterBehavior.IsValid() && BrushAdjusterBehavior->IsBrushBeingAdjusted())
+	if (BrushAdjusterBehavior.IsValid())
 	{
-		return true;
+		if (BrushAdjusterBehavior->IsBrushBeingAdjusted())
+		{
+			return true;
+		}
+
+		// When not in adjustment mode, keep the brush & adjustment origin synchronized with
+		// the brush so that the initial BrushAdjuster HUD display tracks the brush stamp.
+		BrushAdjusterBehavior->OnDragStart(DevicePos.ScreenPosition);
 	}
 	
 	FHitResult OutHit;
diff --git a/Engine/Source/Runtime/InteractiveToolsFramework/Public/BaseTools/BaseBrushTool.h b/Engine/Source/Runtime/InteractiveToolsFramework/Public/BaseTools/BaseBrushTool.h
--- a/Engine/Source/Runtime/InteractiveToolsFramework/Public/BaseTools/BaseBrushTool.h
+++ b/Engine/Source/Runtime/InteractiveToolsFramework/Public/BaseTools/BaseBrushTool.h
@@ -165,8 +165,6 @@ public:
 	
 	// IClickDragBehaviorTarget implementation
 	INTERACTIVETOOLSFRAMEWORK_API virtual FInputRayHit CanBeginClickDragSequence(const FInputDeviceRay& PressPos) override;
-	INTERACTIVETOOLSFRAMEWORK_API virtual void OnClickPress(const FInputDeviceRay& PressPos) override;
-	INTERACTIVETOOLSFRAMEWORK_API virtual void OnClickDrag(const FInputDeviceRay& DragPos) override;
 	
 	// UMeshSurfacePointTool implementation
 	INTERACTIVETOOLSFRAMEWORK_API virtual void OnBeginDrag(const FRay& Ray) override;
@@ -182,6 +180,7 @@ public:
 	TObjectPtr<UBrushBaseProperties> BrushProperties;
 
 	/** Set to true by Tool if user is currently in an active brush stroke*/
+	UE_DEPRECATED(5.5, "The bInBrushStroke property is deprecated. This will now always be true in OnBeginDrag/OnUpdateDrag/OnEndDrag.")
 	UPROPERTY()
 	bool bInBrushStroke = false;
 
@@ -202,7 +201,10 @@ public:
 	INTERACTIVETOOLSFRAMEWORK_API virtual void IncreaseBrushFalloffAction();
 	INTERACTIVETOOLSFRAMEWORK_API virtual void DecreaseBrushFalloffAction();
 
+	PRAGMA_DISABLE_DEPRECATION_WARNINGS
+	UE_DEPRECATED(5.5, "The IsInBrushStroke function is deprecated. This will now always be true in OnBeginDrag/OnUpdateDrag/OnEndDrag.")
 	virtual bool IsInBrushStroke() const { return bInBrushStroke; }
+	PRAGMA_ENABLE_DEPRECATION_WARNINGS
 
 	virtual double GetCurrentBrushRadius() const { return CurrentBrushRadius; }
 	virtual double GetCurrentBrushRadiusLocal() const { return CurrentBrushRadius * WorldToLocalScale; }
