InteractiveToolsFramework: Fixed brush resizing sensitivity when bSpecifyRadius is false.

diff --git a/Engine/Source/Runtime/InteractiveToolsFramework/Private/BaseTools/BaseBrushTool.cpp b/Engine/Source/Runtime/InteractiveToolsFramework/Private/BaseTools/BaseBrushTool.cpp
--- a/Engine/Source/Runtime/InteractiveToolsFramework/Private/BaseTools/BaseBrushTool.cpp
+++ b/Engine/Source/Runtime/InteractiveToolsFramework/Private/BaseTools/BaseBrushTool.cpp
@@ -108,10 +108,10 @@ void UBrushAdjusterInputBehavior::OnDragUpdate(FVector2D InScreenPosition)
 	if (bAdjustingHorizontally)
 	{
 		// adjust brush size based on horizontal mouse drag
-		float NewRadius = StartBrushRadius + HorizontalDelta * (SizeAdjustSpeed * DPIScale * BrushTool->LastBrushStamp.HitResult.Distance);
-		NewRadius = FMath::Max(NewRadius, 0.01f);
 		if (BrushTool->BrushProperties->bSpecifyRadius)
 		{
+			float NewRadius = StartBrushRadius + HorizontalDelta * (SizeAdjustSpeed * DPIScale * BrushTool->LastBrushStamp.HitResult.Distance);
+			NewRadius = FMath::Max(NewRadius, 0.01f);
 			BrushTool->BrushProperties->BrushRadius = NewRadius;
 #if WITH_EDITOR
 			FPropertyChangedEvent PropertyChangedEvent(UBrushBaseProperties::StaticClass()->FindPropertyByName(GET_MEMBER_NAME_CHECKED(UBrushBaseProperties, BrushRadius)));
@@ -120,7 +120,9 @@ void UBrushAdjusterInputBehavior::OnDragUpdate(FVector2D InScreenPosition)
 		}
 		else
 		{
-			BrushTool->BrushProperties->BrushSize = NewRadius;
+			float NewSize = StartBrushRadius + HorizontalDelta * (SizeAdjustSpeed * DPIScale);
+			NewSize = FMath::Clamp(NewSize, 0.01f, 1.0f);
+			BrushTool->BrushProperties->BrushSize = NewSize;
 #if WITH_EDITOR
 			FPropertyChangedEvent PropertyChangedEvent(UBrushBaseProperties::StaticClass()->FindPropertyByName(GET_MEMBER_NAME_CHECKED(UBrushBaseProperties, BrushSize)));
 			BrushTool->BrushProperties->PostEditChangeProperty(PropertyChangedEvent);
