InteractiveToolsFramework: Added DoubleClickBehavior.

Adjusted InputRouter to forward DoubleClick events to as Pressed events. Clients can register DoubleClick behaviors with a higher priority to capture DoubleClicks. This allows existing single click behaviors the opportunity to consume the double click event as a single click without having to add their own double click behaviors.

diff --git a/Engine/Source/Editor/Experimental/EditorInteractiveToolsFramework/Private/EdModeInteractiveToolsContext.cpp b/Engine/Source/Editor/Experimental/EditorInteractiveToolsFramework/Private/EdModeInteractiveToolsContext.cpp
--- a/Engine/Source/Editor/Experimental/EditorInteractiveToolsFramework/Private/EdModeInteractiveToolsContext.cpp
+++ b/Engine/Source/Editor/Experimental/EditorInteractiveToolsFramework/Private/EdModeInteractiveToolsContext.cpp
@@ -933,7 +933,7 @@ bool UModeManagerInteractiveToolsContext::InputKey(FEditorViewportClient* Viewpo
 	// Update the current state, then route result
 	UpdateStateWithoutRoutingInputKey(ViewportClient, Viewport, Key, Event);
 
-	if (Event == IE_Pressed || Event == IE_Released)
+	if (Event == IE_Pressed || Event == IE_Released || Event == IE_DoubleClick)
 	{
 		if (Key.IsMouseButton())
 		{
@@ -960,20 +960,23 @@ bool UModeManagerInteractiveToolsContext::InputKey(FEditorViewportClient* Viewpo
 					ModifierKeys.IsShiftDown(), ModifierKeys.IsAltDown(),
 					ModifierKeys.IsControlDown(), ModifierKeys.IsCommandDown());
 
+				// DoubleClick replaces a Press event. Include DoubleClick in Pressed/Down events so
+				// that it can still be processed by single click behaviors. For explicit double click
+				// handling, clients can listen for the DoubleClick event with a higher capture priority.
+				const bool bIsPressed = (Event == IE_Pressed || Event == IE_DoubleClick);
+				const bool bIsReleased = (Event == IE_Released);
+				const bool bIsDoubleClick = (Event == IE_DoubleClick);
 				if (bIsLeftMouse)
 				{
-					InputState.Mouse.Left.SetStates(
-						(Event == IE_Pressed), (Event == IE_Pressed), (Event == IE_Released));
+					InputState.Mouse.Left.SetStates(bIsPressed, bIsPressed, bIsReleased, bIsDoubleClick);
 				}
 				else if (bIsMiddleMouse)
 				{
-					InputState.Mouse.Middle.SetStates(
-						(Event == IE_Pressed), (Event == IE_Pressed), (Event == IE_Released));
+					InputState.Mouse.Middle.SetStates(bIsPressed, bIsPressed, bIsReleased, bIsDoubleClick);
 				}
 				else
 				{
-					InputState.Mouse.Right.SetStates(
-						(Event == IE_Pressed), (Event == IE_Pressed), (Event == IE_Released));
+					InputState.Mouse.Right.SetStates(bIsPressed, bIsPressed, bIsReleased, bIsDoubleClick);
 				}
 				if (InputRouter->PostInputEvent(InputState))
 				{
diff --git a/Engine/Source/Runtime/InteractiveToolsFramework/Private/BaseBehaviors/AnyButtonInputBehavior.cpp b/Engine/Source/Runtime/InteractiveToolsFramework/Private/BaseBehaviors/AnyButtonInputBehavior.cpp
--- a/Engine/Source/Runtime/InteractiveToolsFramework/Private/BaseBehaviors/AnyButtonInputBehavior.cpp
+++ b/Engine/Source/Runtime/InteractiveToolsFramework/Private/BaseBehaviors/AnyButtonInputBehavior.cpp
@@ -54,6 +54,16 @@ bool UAnyButtonInputBehavior::IsReleased(const FInputDeviceState& input)
 	return false;
 }
 
+bool UAnyButtonInputBehavior::IsDoubleClicked(const FInputDeviceState& Input)
+{
+	if (Input.IsFromDevice(EInputDevices::Mouse))
+	{
+		ActiveDevice = EInputDevices::Mouse;
+		return GetButtonStateFunc(Input).bDoubleClicked;
+	}
+	return false;
+}
+
 FVector2D UAnyButtonInputBehavior::GetClickPoint(const FInputDeviceState& input)
 {
 	if (input.IsFromDevice(EInputDevices::Mouse))
diff --git a/Engine/Source/Runtime/InteractiveToolsFramework/Private/BaseBehaviors/DoubleClickBehavior.cpp b/Engine/Source/Runtime/InteractiveToolsFramework/Private/BaseBehaviors/DoubleClickBehavior.cpp
new file mode 100644
--- /dev/null
+++ b/Engine/Source/Runtime/InteractiveToolsFramework/Private/BaseBehaviors/DoubleClickBehavior.cpp
@@ -0,0 +1,27 @@
+// Copyright Epic Games, Inc. All Rights Reserved.
+
+#include "BaseBehaviors/DoubleClickBehavior.h"
+
+#include UE_INLINE_GENERATED_CPP_BY_NAME(DoubleClickBehavior)
+
+
+UDoubleClickInputBehavior::UDoubleClickInputBehavior()
+{
+	HitTestOnRelease = true;
+}
+
+
+FInputCaptureRequest UDoubleClickInputBehavior::WantsCapture(const FInputDeviceState& Input)
+{
+	if (IsDoubleClicked(Input) && (ModifierCheckFunc == nullptr || ModifierCheckFunc(Input)) )
+	{
+		FInputRayHit HitResult = Target->IsHitByClick(GetDeviceRay(Input));
+		if (HitResult.bHit)
+		{
+			return FInputCaptureRequest::Begin(this, EInputCaptureSide::Any, HitResult.HitDepth);
+		}
+	}
+	return FInputCaptureRequest::Ignore();
+}
+
+
diff --git a/Engine/Source/Runtime/InteractiveToolsFramework/Public/BaseBehaviors/AnyButtonInputBehavior.h b/Engine/Source/Runtime/InteractiveToolsFramework/Public/BaseBehaviors/AnyButtonInputBehavior.h
--- a/Engine/Source/Runtime/InteractiveToolsFramework/Public/BaseBehaviors/AnyButtonInputBehavior.h
+++ b/Engine/Source/Runtime/InteractiveToolsFramework/Public/BaseBehaviors/AnyButtonInputBehavior.h
@@ -47,6 +47,9 @@ public:
 
 	/** @return true if Target Button was released this frame */
 	INTERACTIVETOOLSFRAMEWORK_API virtual bool IsReleased(const FInputDeviceState& input);
+
+	/** @return true if Target Button was double clicked this frame */
+	INTERACTIVETOOLSFRAMEWORK_API virtual bool IsDoubleClicked(const FInputDeviceState& input);
 	
 	/** @return current 2D position of Target Device, or zero if device does not have 2D position */
 	INTERACTIVETOOLSFRAMEWORK_API virtual FVector2D GetClickPoint(const FInputDeviceState& input);
diff --git a/Engine/Source/Runtime/InteractiveToolsFramework/Public/BaseBehaviors/DoubleClickBehavior.h b/Engine/Source/Runtime/InteractiveToolsFramework/Public/BaseBehaviors/DoubleClickBehavior.h
new file mode 100644
--- /dev/null
+++ b/Engine/Source/Runtime/InteractiveToolsFramework/Public/BaseBehaviors/DoubleClickBehavior.h
@@ -0,0 +1,94 @@
+// Copyright Epic Games, Inc. All Rights Reserved.
+
+#pragma once
+
+#include "BehaviorTargetInterfaces.h"
+#include "CoreMinimal.h"
+#include "InputState.h"
+#include "SingleClickBehavior.h"
+#include "Templates/Function.h"
+#include "UObject/ObjectMacros.h"
+#include "UObject/UObjectGlobals.h"
+
+#include "DoubleClickBehavior.generated.h"
+
+class UObject;
+
+
+/**
+ * UDoubleClickInputBehavior implements a standard "button-click"-style input behavior.
+ * The state machine works as follows:
+ *    1) on input-device-button-double-click, hit-test the target. If hit, begin capture
+ *    2) on input-device-button-release, hit-test the target. If hit, call Target::OnClicked(). If not hit, ignore click.
+ *    
+ * The second hit-test is required to allow the click to be "cancelled" by moving away
+ * from the target. This is standard GUI behavior. You can disable this second hit test
+ * using the .HitTestOnRelease property. This is strongly discouraged.
+ *    
+ * The hit-test and on-clicked functions are provided by a IClickBehaviorTarget instance.
+ *
+ * The expected sequence of mouse events for a double click is:
+ *    a. MouseDown
+ *    b. MouseUp
+ *    c. MouseDoubleClick   <-- State machine starts here.
+ *    d. MouseUp
+ */
+UCLASS(MinimalAPI)
+class UDoubleClickInputBehavior : public USingleClickInputBehavior
+{
+	GENERATED_BODY()
+
+public:
+	INTERACTIVETOOLSFRAMEWORK_API UDoubleClickInputBehavior();
+
+	// UInputBehavior implementation
+	INTERACTIVETOOLSFRAMEWORK_API virtual FInputCaptureRequest WantsCapture(const FInputDeviceState& Input) override;
+};
+
+/**
+ * An implementation of UDoubleClickInputBehavior that also implements IClickBehaviorTarget directly, via a set 
+ * of local lambda functions. To use/customize this class, the client replaces the lambda functions with their own.
+ * This avoids having to create a separate IClickBehaviorTarget implementation for trivial use-cases.
+ */
+UCLASS(MinimalAPI)
+class ULocalDoubleClickInputBehavior : public UDoubleClickInputBehavior, public IClickBehaviorTarget
+{
+	GENERATED_BODY()
+protected:
+	using UDoubleClickInputBehavior::Initialize;
+
+public:
+	/** Call this to initialize the class */
+	virtual void Initialize()
+	{
+		this->Initialize(this);
+	}
+
+	/** lambda implementation of IsHitByClick */
+	TUniqueFunction<FInputRayHit(const FInputDeviceRay&)> IsHitByClickFunc = [](const FInputDeviceRay& ClickPos) { return FInputRayHit(); };
+
+	/** lambda implementation of OnClicked */
+	TUniqueFunction<void(const FInputDeviceRay&)> OnClickedFunc = [](const FInputDeviceRay& ClickPos) {};
+
+	/** lambda implementation of OnUpdateModifierState */
+	TUniqueFunction< void(int, bool) > OnUpdateModifierStateFunc = [](int ModifierID, bool bIsOn) {};
+
+public:
+	// IClickBehaviorTarget implementation
+
+	virtual FInputRayHit IsHitByClick(const FInputDeviceRay& ClickPos) override
+	{
+		return IsHitByClickFunc(ClickPos);
+	}
+
+	virtual void OnClicked(const FInputDeviceRay& ClickPos) override
+	{
+		return OnClickedFunc(ClickPos);
+	}
+
+	// IModifierToggleBehaviorTarget implementation
+	virtual void OnUpdateModifierState(int ModifierID, bool bIsOn)
+	{
+		return OnUpdateModifierStateFunc(ModifierID,bIsOn);
+	}
+};
diff --git a/Engine/Source/Runtime/InteractiveToolsFramework/Public/InputState.h b/Engine/Source/Runtime/InteractiveToolsFramework/Public/InputState.h
--- a/Engine/Source/Runtime/InteractiveToolsFramework/Public/InputState.h
+++ b/Engine/Source/Runtime/InteractiveToolsFramework/Public/InputState.h
@@ -160,24 +160,29 @@ struct FDeviceButtonState
 	UPROPERTY(transient, BlueprintReadWrite, Category = DeviceButtonState)
 	bool bReleased;
 
+	/** Was the button double clicked this frame. This should happen only once per "double click" */
+	UPROPERTY(Transient, BlueprintReadWrite, Category = DeviceButtonState)
+	bool bDoubleClicked;
+
 	FDeviceButtonState()
 	{
 		Button = FKey();
-		bPressed = bDown = bReleased = false;
+		bPressed = bDown = bReleased = bDoubleClicked = false;
 	}
 
 	FDeviceButtonState(const FKey& ButtonIn)
 	{
 		Button = ButtonIn;
-		bPressed = bDown = bReleased = false;
+		bPressed = bDown = bReleased = bDoubleClicked = false;
 	}
 
 	/** Update the states of this button */
-	void SetStates(bool bPressedIn, bool bDownIn, bool bReleasedIn)
+	void SetStates(bool bPressedIn, bool bDownIn, bool bReleasedIn, bool bDoubleClickedIn = false)
 	{
 		bPressed = bPressedIn;
 		bDown = bDownIn;
 		bReleased = bReleasedIn;
+		bDoubleClicked = bDoubleClickedIn;
 	}
 };
 
