InteractiveToolsFramework: Fix UGizmoScaledAndUnscaledTransformSources always resetting the unscaled source scale to one instead of keeping whatever it was.

diff --git a/Engine/Source/Runtime/InteractiveToolsFramework/Private/BaseGizmos/TransformSources.cpp b/Engine/Source/Runtime/InteractiveToolsFramework/Private/BaseGizmos/TransformSources.cpp
--- a/Engine/Source/Runtime/InteractiveToolsFramework/Private/BaseGizmos/TransformSources.cpp
+++ b/Engine/Source/Runtime/InteractiveToolsFramework/Private/BaseGizmos/TransformSources.cpp
@@ -89,7 +89,10 @@ void UGizmoScaledAndUnscaledTransformSources::SetTransform(const FTransform& New
 	if (UnscaledTransformSource)
 	{
 		FTransform Unscaled(NewTransform);
-		Unscaled.SetScale3D(FVector::OneVector);
+		// For the unscaled source, we want the scale to remain unaffected, i.e. same as it currently is. Note 
+		//  that we don't just blindly set the scale to 1 because the unscaled source is typically a gizmo component,
+		//  and we don't want to mess up its scale if it was initialized to a particular value for visual effect.
+		Unscaled.SetScale3D(UnscaledTransformSource->GetTransform().GetScale3D());
 		UnscaledTransformSource->SetTransform(Unscaled);
 	}
 
