InteractiveToolsFramework: Allow setting how tool manager switches between tools. Allow tools to customize their preferred shutdown method when switching between tools using an interface.

diff --git a/Engine/Source/Runtime/InteractiveToolsFramework/Private/InteractiveToolManager.cpp b/Engine/Source/Runtime/InteractiveToolsFramework/Private/InteractiveToolManager.cpp
--- a/Engine/Source/Runtime/InteractiveToolsFramework/Private/InteractiveToolManager.cpp
+++ b/Engine/Source/Runtime/InteractiveToolsFramework/Private/InteractiveToolManager.cpp
@@ -4,6 +4,7 @@
 
 #include "Engine/Engine.h"
 #include "InteractiveToolsContext.h"
+#include "InteractiveToolQueryInterfaces.h"
 #include "ContextObjectStore.h"
 
 #include UE_INLINE_GENERATED_CPP_BY_NAME(InteractiveToolManager)
@@ -149,7 +150,36 @@ bool UInteractiveToolManager::ActivateTool(EToolSide Side)
 
 	if (ActiveLeftTool != nullptr)
 	{
-		DeactivateTool(EToolSide::Left, ActiveLeftTool->CanAccept() ? EToolShutdownType::Accept : EToolShutdownType::Completed);
+		// Decide how to shut down the currently active tool based on ToolSwitchMode
+		EToolShutdownType ShutdownType = EToolShutdownType::Completed;
+		switch (ToolSwitchMode)
+		{
+		case EToolManagerToolSwitchMode::AcceptIfAble:
+		case EToolManagerToolSwitchMode::CustomizableAcceptIfAble:
+			ShutdownType = 
+				CanAcceptActiveTool(EToolSide::Left) ? EToolShutdownType::Accept : 
+				CanCancelActiveTool(EToolSide::Left) ? EToolShutdownType::Cancel : EToolShutdownType::Completed;
+			break;
+		case EToolManagerToolSwitchMode::CancelIfAble:
+		case EToolManagerToolSwitchMode::CustomizableCancelIfAble:
+			ShutdownType = CanCancelActiveTool(EToolSide::Left) ? EToolShutdownType::Cancel : EToolShutdownType::Completed;
+			break;
+		default:
+			ensure(false);
+		}
+
+		// If the tool switch mode is set to be customizable, see if the tool wants to customize it
+		if (ToolSwitchMode == EToolManagerToolSwitchMode::CustomizableAcceptIfAble
+			|| ToolSwitchMode == EToolManagerToolSwitchMode::CustomizableCancelIfAble)
+		{
+			if (IInteractiveToolShutdownQueryAPI* ShutdownQueryAPI = Cast<IInteractiveToolShutdownQueryAPI>(ActiveLeftTool))
+			{
+				ShutdownType = ShutdownQueryAPI->GetPreferredShutdownType(
+					IInteractiveToolShutdownQueryAPI::EShutdownReason::SwitchTool, ShutdownType);
+			}
+		}
+
+		DeactivateTool(EToolSide::Left, ShutdownType);
 	}
 
 	if (ActiveLeftBuilder == nullptr || ActivateToolInternal(Side) == false)
@@ -184,11 +214,6 @@ bool UInteractiveToolManager::ActivateTool(EToolSide Side)
 
 bool UInteractiveToolManager::ActivateToolInternal(EToolSide Side)
 {
-	// We'll keep track of whether the last activated tool has dealt with the stored
-	// tool selection, because we want the default behavior to be a clear of the stored
-	// tool selection on the invocation of any tool that doesn't do anything with it.
-	bActiveToolMadeSelectionStoreRequest = false; // TODO: When we support multiple sides, this approach will need adjustment
-
 	// construct input state we will pass to tools
 	FToolBuilderState InputState;
 	QueriesAPI->GetCurrentSelectionState(InputState);
diff --git a/Engine/Source/Runtime/InteractiveToolsFramework/Public/InteractiveToolManager.h b/Engine/Source/Runtime/InteractiveToolsFramework/Public/InteractiveToolManager.h
--- a/Engine/Source/Runtime/InteractiveToolsFramework/Public/InteractiveToolManager.h
+++ b/Engine/Source/Runtime/InteractiveToolsFramework/Public/InteractiveToolManager.h
@@ -42,6 +42,22 @@ enum class EToolChangeTrackingMode
 	FullUndoRedo = 3
 };
 
+/**
+ * Determines how a currently active tool is shutdown when switching to another tool.
+ */
+UENUM()
+enum class EToolManagerToolSwitchMode
+{
+	/* Accept if able, cancel if not (and complete if neither is possible) */
+	AcceptIfAble,
+	/* Cancel if able, complete if not */
+	CancelIfAble,
+	/* Like AcceptIfAble, except a tool can change the shutdown type by implementing IInteractiveToolShutdownQueryAPI */
+	CustomizableAcceptIfAble,
+	/* Like CancelIfAble, except a tool can change the shutdown type by implementing IInteractiveToolShutdownQueryAPI */
+	CustomizableCancelIfAble,
+};
+
 
 /**
  * UInteractiveToolManager allows users of the tools framework to create and operate Tool instances.
@@ -123,7 +139,8 @@ public:
 	INTERACTIVETOOLSFRAMEWORK_API virtual bool CanActivateTool(EToolSide eSide, FString Identifier);
 
 	/**
-	 * Try to activate a new Tool instance on the given Side
+	 * Try to activate a new Tool instance on the given Side. If another tool is already active, it
+	 *  will be shut down according to ToolSwitchMode.
 	 * @param Side which "side" you would like to active the tool on
 	 * @return true if a new Tool instance was created and initialized
 	 */	
@@ -136,6 +153,21 @@ public:
 	 */
 	INTERACTIVETOOLSFRAMEWORK_API virtual bool HasActiveTool(EToolSide Side) const;
 
+	/**
+	 * Get the current setting for how to deactivate an already-active tool if switching directly to
+	 * another tool.
+	 */
+	INTERACTIVETOOLSFRAMEWORK_API virtual EToolManagerToolSwitchMode GetToolSwitchMode() const { return ToolSwitchMode; }
+
+	/**
+	 * Set the behavior for how to deactivate an already-active tool if ActivateTool is called while
+	 * HasActiveTool is true.
+	 * 
+	 * Note: a higher level system could choose to enforce its preferences for this behavior by querying 
+	 * HasActiveTool and making sure to shut down the previous tool before calling ActivateTool.
+	 * However it is convenient to be able to set the behavior at the tool manager level.
+	 */
+	INTERACTIVETOOLSFRAMEWORK_API virtual void SetToolSwitchMode(EToolManagerToolSwitchMode ToolSwitchModeIn) { ToolSwitchMode = ToolSwitchModeIn; }
 
 	/**
 	 * @return true if there are any active tools
@@ -351,18 +383,13 @@ protected:
 	UInteractiveToolBuilder* ActiveRightBuilder;
 
 	EToolChangeTrackingMode ActiveToolChangeTrackingMode;
+	EToolManagerToolSwitchMode ToolSwitchMode = EToolManagerToolSwitchMode::CustomizableAcceptIfAble;
 
 	FString ActiveLeftToolName;
 	FString ActiveRightToolName;
 	bool bInToolShutdown = false;
 
-	/** 
-	 * Tracks whether the last activated tool has made a request to store a tool selection.
-	 * If it hasn't, we'll clear the currently stored tool selection. The reason for this is
-	 * that we generally don't want to keep the old stored tool selection after a new tool 
-	 * invocation, and we don't want to rely on tools to be kind enough to clear it for us
-	 * (though it is better when they do, since they can bundle it into their own undo transaction).
-	 */
+	UE_DEPRECATED(5.5, "bActiveToolMadeSelectionStoreRequest is no longer used.")
 	bool bActiveToolMadeSelectionStoreRequest = false;
 
 	INTERACTIVETOOLSFRAMEWORK_API virtual bool ActivateToolInternal(EToolSide Side);
diff --git a/Engine/Source/Runtime/InteractiveToolsFramework/Public/InteractiveToolQueryInterfaces.h b/Engine/Source/Runtime/InteractiveToolsFramework/Public/InteractiveToolQueryInterfaces.h
--- a/Engine/Source/Runtime/InteractiveToolsFramework/Public/InteractiveToolQueryInterfaces.h
+++ b/Engine/Source/Runtime/InteractiveToolsFramework/Public/InteractiveToolQueryInterfaces.h
@@ -3,6 +3,7 @@
 #pragma once
 
 #include "CoreMinimal.h"
+#include "InteractiveTool.h" // EToolShutdownType
 #include "UObject/Interface.h"
 #include "InteractiveToolQueryInterfaces.generated.h"
 
@@ -129,6 +130,65 @@ public:
 
 };
 
+// UInterface for IInteractiveToolShutdownQueryAPI
+UINTERFACE(MinimalAPI)
+class UInteractiveToolShutdownQueryAPI : public UInterface
+{
+	GENERATED_BODY()
+};
+
+/**
+ * Allows an interactive tool to influence the way it is shut down, if the tool host supports
+ * that kind of customization. This can be helpful, for example, if your tool prefers a specific
+ * shutdown type in various situations, but a tool can't rely on this interface being queried unless
+ * it knows that it will only be used by systems that respect it. A simple interactive tools framework
+ * context implementation does not need to bother querying the tool on its preferences (aside from the
+ * already existing CanAccept method on the actual tool object).
+ * 
+ * Note that there are different systems that might choose to query this interface. The tool manager
+ * might want to query it if its ToolSwitchMode is set to be customizable, or the mode or mode toolkit
+ * might query it when shutting down the tool in various situations.
+ */
+class IInteractiveToolShutdownQueryAPI
+{
+	GENERATED_BODY()
+public:
+	enum class EShutdownReason
+	{
+		// Another tool is being activated without the user having explicitly shut down this one.
+		SwitchTool
+
+		//~ We'll add something along these lines once we start querying in other shutdown situations:
+		//~ // The ESC key (or some equivalent) was hit to exit the tool
+		//~ , Escape
+		//~ // The mode was shut down
+		//~ , ModeExit
+	};
+
+	/**
+	 * Given a shutdown situation, tells what kind of shutdown method the tool might prefer.
+	 * @param ShutdownReason Information about the shutdown type
+	 * @param StandardShutdownType Shutdown type that the host would use if it weren't giving the tool a
+	 *   chance to customize the behavior.
+	 * @return Shutdown type that the tool would prefer that the host use in this situation.
+	 */
+	virtual EToolShutdownType GetPreferredShutdownType(EShutdownReason ShutdownReason, EToolShutdownType StandardShutdownType) const
+	{
+		return StandardShutdownType;
+	}
+
+	//~ Not currently supported, but someday we might let tools query the user with an "are you sure"
+	//~ type of message, and allow the tool to ask to NOT be shut down after all. This is one way
+	//~ we might implement the ability for tools to ask for this:
+	//~ struct FShutdownUserQuery
+	//~ {
+	//~ 	FString MessageToUser;
+	//~ 	EAppMsgType::Type QueryType;
+	//~ 	TFunction<void(EAppReturnType::Type UserResponse, bool& bStillShutdownOut, EToolShutdownType& ShutdownTypeOut)> ResponseHandler;
+	//~ };
+	//~ TOptional<FShutdownUserQuery> GetShutdownUserQuery() { return TOptional<FShutdownUserQuery>(); }
+};
+
 
 // UInterface for IInteractiveToolExclusiveToolAPI
 UINTERFACE(MinimalAPI)
