InteractiveToolsFramework: Custom static mesh components and architecture for gizmo customization.

#virtualized

diff --git a/Engine/Source/Runtime/InteractiveToolsFramework/InteractiveToolsFramework.Build.cs b/Engine/Source/Runtime/InteractiveToolsFramework/InteractiveToolsFramework.Build.cs
--- a/Engine/Source/Runtime/InteractiveToolsFramework/InteractiveToolsFramework.Build.cs
+++ b/Engine/Source/Runtime/InteractiveToolsFramework/InteractiveToolsFramework.Build.cs
@@ -25,7 +25,8 @@ public class InteractiveToolsFramework : ModuleRules
 			{
 				"Engine",
                 "RHI",
-				"GeometryCore"
+				"GeometryCore",
+				"RenderCore" // for gizmo components
 				//"Slate",
 				//"SlateCore",
 				// ... add private dependencies that you statically link with here ...	
diff --git a/Engine/Source/Runtime/InteractiveToolsFramework/Private/BaseGizmos/GizmoPrivateUtil.cpp b/Engine/Source/Runtime/InteractiveToolsFramework/Private/BaseGizmos/GizmoPrivateUtil.cpp
new file mode 100644
--- /dev/null
+++ b/Engine/Source/Runtime/InteractiveToolsFramework/Private/BaseGizmos/GizmoPrivateUtil.cpp
@@ -0,0 +1,26 @@
+// Copyright Epic Games, Inc. All Rights Reserved.
+
+#include "BaseGizmos/GizmoPrivateUtil.h"
+
+#include "BaseGizmos/GizmoViewContext.h"
+#include "ContextObjectStore.h"
+#include "InteractiveGizmoManager.h"
+
+UGizmoViewContext* UE::GizmoUtil::GetGizmoViewContext(UInteractiveGizmoManager* GizmoManager)
+{
+	if (!ensure(GizmoManager))
+	{
+		return nullptr;
+	}
+
+	UContextObjectStore* ContextObjectStore = GizmoManager->GetContextObjectStore();
+	if (!ensure(ContextObjectStore))
+	{
+		return nullptr;
+	}
+
+	UGizmoViewContext* GizmoViewContext = ContextObjectStore->FindContext<UGizmoViewContext>();
+	ensure(GizmoViewContext);
+
+	return GizmoViewContext;
+}
diff --git a/Engine/Source/Runtime/InteractiveToolsFramework/Private/BaseGizmos/GizmoPrivateUtil.h b/Engine/Source/Runtime/InteractiveToolsFramework/Private/BaseGizmos/GizmoPrivateUtil.h
new file mode 100644
--- /dev/null
+++ b/Engine/Source/Runtime/InteractiveToolsFramework/Private/BaseGizmos/GizmoPrivateUtil.h
@@ -0,0 +1,28 @@
+// Copyright Epic Games, Inc. All Rights Reserved.
+
+#pragma once
+
+class UGizmoViewContext;
+class UInteractiveGizmoManager;
+
+namespace UE::GizmoUtil
+{
+	struct FTransformSubGizmoCommonParams;
+	struct FTransformSubGizmoSharedState;
+}
+
+/**
+ * This file holds implementation helpers that don't necessarily need exposing. If eventually needed,
+ *  we can move some of these into GizmoUtil or TransformSubGizmoUtil
+ */
+
+namespace UE::GizmoUtil 
+{
+
+	/**
+	 * Simple helper that gets the gizmo view context out of the context object store associated with a gizmo manager.
+	 * Fires ensures if it does not find the expected objects along the way.
+	 */
+	UGizmoViewContext* GetGizmoViewContext(UInteractiveGizmoManager* GizmoManager);
+
+}// end UE::GizmoUtil
\ No newline at end of file
diff --git a/Engine/Source/Runtime/InteractiveToolsFramework/Private/BaseGizmos/GizmoRenderingUtil.cpp b/Engine/Source/Runtime/InteractiveToolsFramework/Private/BaseGizmos/GizmoRenderingUtil.cpp
--- a/Engine/Source/Runtime/InteractiveToolsFramework/Private/BaseGizmos/GizmoRenderingUtil.cpp
+++ b/Engine/Source/Runtime/InteractiveToolsFramework/Private/BaseGizmos/GizmoRenderingUtil.cpp
@@ -1,12 +1,16 @@
 // Copyright Epic Games, Inc. All Rights Reserved.
 
 #include "BaseGizmos/GizmoRenderingUtil.h"
+
+#include "BaseGizmos/GizmoPrivateUtil.h" // GetGizmoViewContext
 #include "BaseGizmos/GizmoViewContext.h"
+#include "BaseGizmos/ViewAdjustedStaticMeshGizmoComponent.h"
 #include "RHI.h"
+#include "Materials/Material.h"
+#include "Materials/MaterialInstanceDynamic.h"
 #include "Math/UnrealMathUtility.h"
 #include "Math/NumericLimits.h"
 
-
 // yuck global value set by Editor
 static const FSceneView* GlobalCurrentSceneView = nullptr;
 
@@ -40,59 +44,99 @@ namespace GizmoRenderingUtilLocals
 	}
 
 	// duplicates FSceneView::WorldToPixel but in double where possible (unfortunately WorldToScreen still in float)
-	template <typename SceneViewOrGizmoViewContext>
-	static FVector2D WorldToPixelDouble(const SceneViewOrGizmoViewContext* View, const FVector& Location)
+	static FVector2D WorldToPixelDouble(const UE::GizmoRenderingUtil::ISceneViewInterface* View, const FVector& Location)
 	{
 		FVector4 ScreenPoint = View->WorldToScreen(Location);
 
 		double InvW = (ScreenPoint.W > 0.0 ? 1.0 : -1.0) / (double)ScreenPoint.W;
 		double Y = (GProjectionSignY > 0.0) ? (double)ScreenPoint.Y : 1.0 - (double)ScreenPoint.Y;
 
-		const FIntRect& UnscaledViewRect = View->UnscaledViewRect;
+		const FIntRect& UnscaledViewRect = View->GetUnscaledViewRect();
 		double PosX = (double)UnscaledViewRect.Min.X + (0.5 + (double)ScreenPoint.X * 0.5 * InvW) * (double)UnscaledViewRect.Width();
 		double PosY = (double)UnscaledViewRect.Min.Y + (0.5 - Y * 0.5 * InvW) * (double)UnscaledViewRect.Height();
 
 		return FVector2D((float)PosX, (float)PosY);
 	}
 
-	// Shared code between the GizmoViewContext and FSceneView versions of CalculateLocalPixelToWorldScale. We don't expose
-	// this template itself to the user since it's not entirely intuitive that this works on two very different types.
-	template <typename SceneViewOrGizmoViewContext>
-	float CalculateLocalPixelToWorldScale(
-		const SceneViewOrGizmoViewContext* View,
-		const FVector& Location)
+	// This matches "CurrentColor" in FWidget::Widget()
+	FColor HoverColor = FColor::Yellow;
+}
+
+UViewAdjustedStaticMeshGizmoComponent* UE::GizmoRenderingUtil::CreateDefaultMaterialGizmoMeshComponent(
+	UStaticMesh* Mesh, UGizmoViewContext* GizmoViewContext, UObject* OwnerComponentOrActor, 
+	const FLinearColor& Color, bool bAddHoverMaterial)
+{
+	using namespace GizmoRenderingUtilLocals;
+
+	if (!ensureMsgf(OwnerComponentOrActor, TEXT("CreateDefaultMaterialGizmoMeshComponent: Need owner component or actor to create component.")))
+	{
+		return nullptr;
+	}
+	UViewAdjustedStaticMeshGizmoComponent* Component = NewObject<UViewAdjustedStaticMeshGizmoComponent>(OwnerComponentOrActor);
+	Component->SetStaticMesh(Mesh);
+
+	Component->SetGizmoViewContext(GizmoViewContext);
+	Component->TranslucencySortPriority = GizmoRenderingUtil::GIZMO_TRANSLUCENCY_SORT_PRIORITY;
+	// Used by the default material to be able to be occluded by other gizmo elements
+	Component->bRenderCustomDepth = true;
+	// Not sure that this actually gets respected in any way for non-PDI calls, but just in case
+	Component->DepthPriorityGroup = SDPG_Foreground;
+
+	Component->SetAllMaterials(GizmoRenderingUtil::GetDefaultGizmoComponentMaterial(Color, Component));
+	if (bAddHoverMaterial)
 	{
-		// To calculate this scale at Location, we project Location to screen and also project a second
-		// point at a small distance from Location in a camera-perpendicular plane, then measure 2D/3D distance ratio. 
-		// However, because some of the computations are done in float, there will be enormous numerical error 
-		// when the camera is very far from the location if the distance is relatively small. The "W" value
-		// below gives us a sense of this distance, so we make the offset relative to that
-		// (this does do one redundant WorldToScreen)
-		FVector4 LocationScreenPoint = View->WorldToScreen(Location);
-		double OffsetDelta = LocationScreenPoint.W * 0.01;
-
-		FVector2D PixelA = WorldToPixelDouble(View, Location);
-		FVector OffsetPointWorld = Location + OffsetDelta * View->GetViewRight() + OffsetDelta * View->GetViewUp();
-		FVector2D PixelB = WorldToPixelDouble(View, OffsetPointWorld);
-
-		double PixelDeltaSqr = VectorDifferenceSqr(PixelA, PixelB);
-		double WorldDeltaSqr = VectorDifferenceSqr(Location, OffsetPointWorld);
-		return (float)(sqrt(WorldDeltaSqr / PixelDeltaSqr));
+		Component->SetHoverOverrideMaterial(UE::GizmoRenderingUtil::GetDefaultGizmoComponentMaterial(HoverColor, Component));
 	}
+
+	return Component;
 }
 
+UViewAdjustedStaticMeshGizmoComponent* UE::GizmoRenderingUtil::CreateDefaultMaterialGizmoMeshComponent(
+	UStaticMesh* Mesh, UInteractiveGizmoManager* GizmoManager, UObject* OwnerComponentOrActor, 
+	const FLinearColor& Color, bool bAddHoverMaterial)
+{
+	return CreateDefaultMaterialGizmoMeshComponent(Mesh, UE::GizmoUtil::GetGizmoViewContext(GizmoManager), 
+		OwnerComponentOrActor, Color, bAddHoverMaterial);
+}
+
+
 float GizmoRenderingUtil::CalculateLocalPixelToWorldScale(
 	const FSceneView* View,
 	const FVector& Location)
 {
-	return GizmoRenderingUtilLocals::CalculateLocalPixelToWorldScale(View, Location);
+	using namespace GizmoRenderingUtilLocals;
+
+	if (!View)
+	{
+		return 0.0f;
+	}
+
+	UE::GizmoRenderingUtil::FSceneViewWrapper Wrapper(*View);
+	return CalculateLocalPixelToWorldScale(&Wrapper, Location);
 }
 
 float GizmoRenderingUtil::CalculateLocalPixelToWorldScale(
-	const UGizmoViewContext* ViewContext,
+	const UE::GizmoRenderingUtil::ISceneViewInterface* View, 
 	const FVector& Location)
 {
-	return GizmoRenderingUtilLocals::CalculateLocalPixelToWorldScale(ViewContext, Location);
+	using namespace GizmoRenderingUtilLocals;
+
+	// To calculate this scale at Location, we project Location to screen and also project a second
+	// point at a small distance from Location in a camera-perpendicular plane, then measure 2D/3D distance ratio. 
+	// However, because some of the computations are done in float, there will be enormous numerical error 
+	// when the camera is very far from the location if the distance is relatively small. The "W" value
+	// below gives us a sense of this distance, so we make the offset relative to that
+	// (this does do one redundant WorldToScreen)
+	FVector4 LocationScreenPoint = View->WorldToScreen(Location);
+	double OffsetDelta = LocationScreenPoint.W * 0.01;
+
+	FVector2D PixelA = WorldToPixelDouble(View, Location);
+	FVector OffsetPointWorld = Location + OffsetDelta * View->GetViewRight() + OffsetDelta * View->GetViewUp();
+	FVector2D PixelB = WorldToPixelDouble(View, OffsetPointWorld);
+
+	double PixelDeltaSqr = VectorDifferenceSqr(PixelA, PixelB);
+	double WorldDeltaSqr = VectorDifferenceSqr(Location, OffsetPointWorld);
+	return (float)(sqrt(WorldDeltaSqr / PixelDeltaSqr));
 }
 
 float GizmoRenderingUtil::CalculateViewDependentScaleAndFlatten(
@@ -139,3 +183,33 @@ float GizmoRenderingUtil::CalculateViewDependentScaleAndFlatten(
 	}
 	return UniformScale;
 }
+
+UMaterialInterface* UE::GizmoRenderingUtil::GetDefaultGizmoComponentMaterial(const FLinearColor& Color, UObject* Outer)
+{
+	UMaterialInterface* Material = LoadObject<UMaterial>(nullptr, TEXT("/MeshModelingToolsetExp/Materials/GizmoComponentMaterial"));
+	if (Material == nullptr)
+	{
+		return nullptr;
+	}
+
+	UMaterialInstanceDynamic* MatInstance = UMaterialInstanceDynamic::Create(Material, Outer);
+
+	MatInstance->SetVectorParameterValue(TEXT("GizmoColor"), Color);
+	return MatInstance;
+}
+
+FLinearColor UE::GizmoRenderingUtil::GetDefaultAxisColor(EAxis::Type Axis)
+{
+	// The below colors come from FWidget::FWidget() and elsewhere
+	switch (Axis)
+	{
+	case EAxis::X:
+		return FLinearColor(0.594f, 0.0197f, 0.0f);
+	case EAxis::Y:
+		return FLinearColor(0.1349f, 0.3959f, 0.0f);
+	case EAxis::Z:
+		return FLinearColor(0.0251f, 0.207f, 0.85f);
+	default:
+		return FLinearColor::White;
+	}
+}
\ No newline at end of file
diff --git a/Engine/Source/Runtime/InteractiveToolsFramework/Private/BaseGizmos/ViewAdjustedStaticMeshGizmoComponent.cpp b/Engine/Source/Runtime/InteractiveToolsFramework/Private/BaseGizmos/ViewAdjustedStaticMeshGizmoComponent.cpp
new file mode 100644
--- /dev/null
+++ b/Engine/Source/Runtime/InteractiveToolsFramework/Private/BaseGizmos/ViewAdjustedStaticMeshGizmoComponent.cpp
@@ -0,0 +1,357 @@
+// Copyright Epic Games, Inc. All Rights Reserved.
+
+#include "BaseGizmos/ViewAdjustedStaticMeshGizmoComponent.h"
+
+#include "BaseGizmos/GizmoViewContext.h"
+#include "BaseGizmos/CombinedTransformGizmo.h"
+#include "BaseGizmos/ViewBasedTransformAdjusters.h"
+#include "BaseGizmos/GizmoRenderingUtil.h"
+#include "StaticMeshResources.h"
+#include "StaticMeshSceneProxy.h"
+#include "StaticMeshSceneProxyDesc.h"
+
+namespace ViewAdjustedStaticMeshGizmoComponentLocals
+{
+	// Overriding a function here is necessary to be able to properly update the used materials to include
+	//  the hover material.
+	class FViewAdjustedStaticMeshSceneProxyDesc : public FStaticMeshSceneProxyDesc
+	{
+	public:
+		FViewAdjustedStaticMeshSceneProxyDesc(const UViewAdjustedStaticMeshGizmoComponent* Component)
+			: FStaticMeshSceneProxyDesc(Component)
+		{
+		}
+
+		virtual void GetUsedMaterials(TArray<UMaterialInterface*>&OutMaterials, bool bGetDebugMaterials = false) const override
+		{
+			FStaticMeshSceneProxyDesc::GetUsedMaterials(OutMaterials);
+			if (UViewAdjustedStaticMeshGizmoComponent* CastComponent = Cast<UViewAdjustedStaticMeshGizmoComponent>(Component))
+			{
+				if (UMaterialInterface* HoverMaterial = CastComponent->GetHoverOverrideMaterial())
+				{
+					OutMaterials.Add(HoverMaterial);
+				}
+			}
+		}
+	};
+
+	class FViewAdjustedStaticMeshGizmoComponentProxy : public FStaticMeshSceneProxy
+	{
+	public:
+		FViewAdjustedStaticMeshGizmoComponentProxy(UViewAdjustedStaticMeshGizmoComponent* Component)
+			: FStaticMeshSceneProxy(FViewAdjustedStaticMeshSceneProxyDesc(Component), false)
+			, TransformAdjuster(Component->GetTransformAdjuster())
+			, HoverOverrideMaterial(Component->GetHoverOverrideMaterial())
+			, bHovered(Component->IsBeingHovered())
+			, bHiddenByInteraction(Component->IsHiddenByInteraction())
+		{
+		}
+
+		void SetIsHovered(bool bHoveredIn)
+		{
+			bHovered = bHoveredIn;
+		}
+
+		void SetIsHiddenByInteraction(bool bIsHidden)
+		{
+			bHiddenByInteraction = bIsHidden;
+		}
+
+		void SetTransformAdjuster(TSharedPtr<UE::GizmoRenderingUtil::IViewBasedTransformAdjuster> TransformAdjusterIn)
+		{
+			TransformAdjuster = TransformAdjusterIn;
+		}
+
+		virtual void GetDynamicMeshElements(const TArray<const FSceneView*>& Views,
+			const FSceneViewFamily& ViewFamily, uint32 VisibilityMap, FMeshElementCollector& Collector) const override
+		{
+			if (bHiddenByInteraction)
+			{
+				return;
+			}
+
+			// For the most part, the below is modeled off the FStaticMeshSceneProxy version of this method, with 
+			//  various things cut out (debug view modes, etc) and some deep nesting turned out into early outs.
+			for (int32 ViewIndex = 0; ViewIndex < Views.Num(); ViewIndex++)
+			{
+				const FSceneView* View = Views[ViewIndex];
+
+				if (!(IsShown(View) && (VisibilityMap & (1 << ViewIndex))))
+				{
+					continue;
+				}
+
+				FFrozenSceneViewMatricesGuard FrozenMatricesGuard(*const_cast<FSceneView*>(Views[ViewIndex]));
+
+				// We can calculate our adjusted transform at this point, now that we have the view
+				const TUniformBuffer<FPrimitiveUniformShaderParameters>* AdjustedTransformBuffer = nullptr;
+				bool bAdjustedTransformDeterminantIsNegative = IsLocalToWorldDeterminantNegative();
+				if (TransformAdjuster.IsValid())
+				{
+					FTransform NewTransform = TransformAdjuster->GetAdjustedComponentToWorld_RenderThread(
+						UE::GizmoRenderingUtil::FSceneViewWrapper(*View), FTransform(GetLocalToWorld()));
+
+					FMatrix NewTransformMatrix = NewTransform.ToMatrixWithScale();
+					bAdjustedTransformDeterminantIsNegative = (NewTransformMatrix.Determinant() < 0);
+
+					const FBoxSphereBounds3d& OriginalLocalBounds = GetLocalBounds();
+
+					// This way of setting the transform is copied from FTriangleSetSceneProxy::GetDynamicMeshElements
+					FDynamicPrimitiveUniformBuffer& DynamicPrimitiveUniformBuffer = Collector.AllocateOneFrameResource<FDynamicPrimitiveUniformBuffer>();
+					DynamicPrimitiveUniformBuffer.Set(Collector.GetRHICommandList(),
+						NewTransformMatrix, NewTransformMatrix,
+						GetLocalBounds().TransformBy(NewTransformMatrix), GetLocalBounds(),
+						/*bReceivesDecals*/ true, /*bHasPrecomputedVolumetricLightmap*/ false,
+						AlwaysHasVelocity());
+
+					AdjustedTransformBuffer = &DynamicPrimitiveUniformBuffer.UniformBuffer;
+				}//end getting adjusted transform ready
+
+				FLODMask LODMask = GetLODMask(View);
+
+				for (int32 LODIndex = 0; LODIndex < RenderData->LODResources.Num(); LODIndex++)
+				{
+					if (!(LODMask.ContainsLOD(LODIndex) && LODIndex >= ClampedMinLOD))
+					{
+						continue;
+					}
+
+					const FStaticMeshLODResources& LODModel = RenderData->LODResources[LODIndex];
+
+					for (int32 SectionIndex = 0; SectionIndex < LODModel.Sections.Num(); SectionIndex++)
+					{
+						const int32 NumBatches = GetNumMeshBatches();
+
+						for (int32 BatchIndex = 0; BatchIndex < NumBatches; BatchIndex++)
+						{
+							FMeshBatch& MeshBatch = Collector.AllocateMesh();
+
+							// This selection and hit proxy id setting seems unneeded but we'll just keep it
+							bool bSectionIsSelected = false;
+#if WITH_EDITOR
+							if (GIsEditor)
+							{
+								const FLODInfo::FSectionInfo& Section = LODs[LODIndex].Sections[SectionIndex];
+
+								bSectionIsSelected = Section.bSelected;
+								MeshBatch.BatchHitProxyId = Section.HitProxy ? Section.HitProxy->Id : FHitProxyId();
+							}
+#endif
+
+							if (!GetMeshElement(LODIndex, BatchIndex, SectionIndex,
+								GetStaticDepthPriorityGroup(), bSectionIsSelected,
+								// Not sure what this bAllowPreCulledIndices parameter is, but this is what FStaticMeshSceneProxy does
+								true,
+								MeshBatch))
+							{
+								continue;
+							}
+
+							// Seems like there is only ever one of these...
+							FMeshBatchElement& BatchElement = MeshBatch.Elements[0];
+
+							// Update the portions affected by the adjusted transform determinant
+							const bool bUseReversedIndices =
+								// based off FStaticMeshSceneProxy::ShouldRenderBackFaces()
+								(bAdjustedTransformDeterminantIsNegative != bReverseCulling) &&
+								RenderData->LODResources[LODIndex].bHasReversedDepthOnlyIndices;
+							MeshBatch.ReverseCulling = IsReversedCullingNeeded(bUseReversedIndices);
+							BatchElement.IndexBuffer = bUseReversedIndices 
+								? &LODModel.AdditionalIndexBuffers->ReversedIndexBuffer 
+								: &LODModel.IndexBuffer;
+
+							// Gizmos probably don't want to be affected by view mode?
+							MeshBatch.bCanApplyViewModeOverrides = false;
+
+							// This is where we bind our adjusted transform:
+							if (AdjustedTransformBuffer)
+							{
+								BatchElement.PrimitiveUniformBufferResource = AdjustedTransformBuffer;
+							}
+							
+							// Apply hover override material
+							if (bHovered && HoverOverrideMaterial)
+							{
+								MeshBatch.MaterialRenderProxy = HoverOverrideMaterial->GetRenderProxy();
+							}
+
+							Collector.AddMesh(ViewIndex, MeshBatch);
+							INC_DWORD_STAT_BY(STAT_StaticMeshTriangles, MeshBatch.GetNumPrimitives());
+
+						}// for each MeshBatch
+					}// for each mesh section
+				}// for each LOD
+			}// for each view
+		}//GetDynamicMeshElements()
+
+		virtual FPrimitiveViewRelevance GetViewRelevance(const FSceneView* View) const override
+		{
+			FPrimitiveViewRelevance Result = FStaticMeshSceneProxy::GetViewRelevance(View);
+			Result.bDynamicRelevance = true;
+			Result.bStaticRelevance = false;
+			Result.bShadowRelevance = false;
+			
+			MaterialRelevance.SetPrimitiveViewRelevance(Result);
+
+			return Result;
+		}
+
+		virtual void DrawStaticElements(FStaticPrimitiveDrawInterface* PDI) override {}
+
+		virtual bool CanBeOccluded() const override 
+		{
+			// If we're using a transform adjuster, we're going to say that we can't be occluded because in the common case
+			//  of keeping the component a constant view size, it can be arbitrarily large as we move away from it. This prevents
+			//  the component from occluding itself too.
+			// This isn't actually necessary if we're using a non-depth-tested material.
+			return !TransformAdjuster.IsValid() && FStaticMeshSceneProxy::CanBeOccluded();
+		};
+
+		private:
+
+		TSharedPtr<UE::GizmoRenderingUtil::IViewBasedTransformAdjuster> TransformAdjuster;
+		UMaterialInterface* HoverOverrideMaterial = nullptr;
+		bool bHovered = false;
+		// It is tempting to use the visibility of the component to hide it during interaction, but that turns out to be
+		//  problematic because other things affect the visibility- for example the TRS gizmo constantly updates visibility
+		//  depending on the current gizmo mode (translate/rotate/scale). Instead, we want this setting to be an extra flag
+		//  that forces invisibility. 
+		bool bHiddenByInteraction = false;
+	};
+
+}
+
+
+bool UViewAdjustedStaticMeshGizmoComponent::LineTraceComponent(FHitResult& OutHit, const FVector Start, const FVector End, const FCollisionQueryParams& Params)
+{
+	// If needed, update the physics data, then do the line trace
+	if (ensure(GizmoViewContext) && TransformAdjuster.IsValid())
+	{
+		FTransform OriginalComponentToWorld = GetComponentToWorld();
+		FTransform AdjustedComponentToWorld = TransformAdjuster->GetAdjustedComponentToWorld(*GizmoViewContext, OriginalComponentToWorld);
+		if (!OriginalComponentToWorld.Equals(AdjustedComponentToWorld))
+		{
+			BodyInstance.SetBodyTransform(AdjustedComponentToWorld, ETeleportType::None);
+			BodyInstance.UpdateBodyScale(AdjustedComponentToWorld.GetScale3D());
+		}
+	}
+
+	return Super::LineTraceComponent(OutHit, Start, End, Params);
+}
+
+FPrimitiveSceneProxy* UViewAdjustedStaticMeshGizmoComponent::CreateStaticMeshSceneProxy(Nanite::FMaterialAudit& NaniteMaterials, bool bCreateNanite)
+{
+	return ::new ViewAdjustedStaticMeshGizmoComponentLocals::FViewAdjustedStaticMeshGizmoComponentProxy(this);
+}
+
+FMaterialRelevance UViewAdjustedStaticMeshGizmoComponent::GetMaterialRelevance(ERHIFeatureLevel::Type InFeatureLevel) const
+{
+	FMaterialRelevance Result = Super::GetMaterialRelevance(InFeatureLevel);
+	if (HoverOverrideMaterial)
+	{
+		Result |= HoverOverrideMaterial->GetRelevance_Concurrent(InFeatureLevel);
+	}
+	return Result;
+}
+
+void UViewAdjustedStaticMeshGizmoComponent::SetTransformAdjuster(TSharedPtr<UE::GizmoRenderingUtil::IViewBasedTransformAdjuster> Adjuster)
+{
+	TransformAdjuster = Adjuster;
+	MarkRenderStateDirty();
+}
+
+void UViewAdjustedStaticMeshGizmoComponent::SetAllMaterials(UMaterialInterface* Material)
+{
+	for (int i = 0; i < GetNumMaterials(); ++i)
+	{
+		SetMaterial(i, Material);
+	}
+}
+
+void UViewAdjustedStaticMeshGizmoComponent::SetHoverOverrideMaterial(UMaterialInterface* Material)
+{
+	HoverOverrideMaterial = Material;
+	MarkRenderStateDirty();
+}
+
+void UViewAdjustedStaticMeshGizmoComponent::SetSubstituteInteractionComponent(UPrimitiveComponent* NewComponent,
+	const FTransform& RelativeTransform)
+{
+	if (SubstituteInteractionComponent == NewComponent)
+	{
+		return;
+	}
+
+	if (SubstituteInteractionComponent)
+	{
+		if (SubstituteInteractionComponent->IsRegistered())
+		{
+			SubstituteInteractionComponent->UnregisterComponent();
+		}
+		SubstituteInteractionComponent->DetachFromComponent(FDetachmentTransformRules::KeepWorldTransform);
+	}
+
+	SubstituteInteractionComponent = NewComponent;
+
+	if (!NewComponent)
+	{
+		// If we're clearing the substitute component, we're done
+		return;
+	}
+
+	if (NewComponent->IsRegistered())
+	{
+		NewComponent->UnregisterComponent();
+	}
+	NewComponent->AttachToComponent(this, FAttachmentTransformRules::KeepWorldTransform);
+	NewComponent->SetRelativeTransform(RelativeTransform);
+	NewComponent->RegisterComponent();
+
+	UpdateInteractingState(false);
+}
+
+void UViewAdjustedStaticMeshGizmoComponent::UpdateInteractingState(bool bInteracting)
+{
+	// Note: we don't early out if we're already interacting because we want to be able to call
+	// this function to update the scene proxy.
+	bInteracted = bInteracting;
+
+	if (SubstituteInteractionComponent)
+	{
+		SubstituteInteractionComponent->SetVisibility(bInteracting);
+	}
+
+	if (FPrimitiveSceneProxy* Proxy = GetSceneProxy())
+	{
+		static_cast<ViewAdjustedStaticMeshGizmoComponentLocals::FViewAdjustedStaticMeshGizmoComponentProxy*>(Proxy)
+			->SetIsHiddenByInteraction(bInteracted && SubstituteInteractionComponent);
+	}
+}
+
+void UViewAdjustedStaticMeshGizmoComponent::UpdateHoverState(bool bHoveringIn)
+{
+	if (bHoveringIn != bHovered)
+	{
+		bHovered = bHoveringIn;
+		if (FPrimitiveSceneProxy* Proxy = GetSceneProxy())
+		{
+			static_cast<ViewAdjustedStaticMeshGizmoComponentLocals::FViewAdjustedStaticMeshGizmoComponentProxy*>(Proxy)
+				->SetIsHovered(bHovered);
+		}
+	}
+}
+
+void UViewAdjustedStaticMeshGizmoComponent::UpdateWorldLocalState(bool bWorldIn)
+{
+	if (TransformAdjuster.IsValid())
+	{
+		TransformAdjuster->UpdateWorldLocalState(bWorldIn);
+	}
+	// If able to, forward this information to the substitute component
+	// Tempting to only do this if we're interacting, but what if we get the update just once, and
+	//  never update the substitute...
+	if (IGizmoBaseComponentInterface* CastSubstituteComp = Cast<IGizmoBaseComponentInterface>(SubstituteInteractionComponent))
+	{
+		CastSubstituteComp->UpdateWorldLocalState(bWorldIn);
+	}
+}
diff --git a/Engine/Source/Runtime/InteractiveToolsFramework/Private/BaseGizmos/ViewBasedTransformAdjusters.cpp b/Engine/Source/Runtime/InteractiveToolsFramework/Private/BaseGizmos/ViewBasedTransformAdjusters.cpp
new file mode 100644
--- /dev/null
+++ b/Engine/Source/Runtime/InteractiveToolsFramework/Private/BaseGizmos/ViewBasedTransformAdjusters.cpp
@@ -0,0 +1,176 @@
+// Copyright Epic Games, Inc. All Rights Reserved.
+
+#include "BaseGizmos/ViewBasedTransformAdjusters.h"
+
+#include "BaseGizmos/GizmoRenderingUtil.h"
+#include "BaseGizmos/ViewAdjustedStaticMeshGizmoComponent.h"
+#include "FrameTypes.h"
+
+namespace ViewBasedTransformAdjustersLocals
+{
+	double GetComponentToGizmoScaling(
+		const UE::GizmoRenderingUtil::ISceneViewInterface& View,
+		const FTransform& GizmoToWorld)
+	{
+		// It might seem that all we want is CalculateLocalPixelToWorldScale(View, CurrentComponentToWorld.GetLocation());
+		//  However we're in a weird situation where scaling gets applied around the gizmo origin, not the component (to
+		//  preserve component positioning within the gizmo); moreover we want different sub components of the same gizmo
+		//  to be scaled the same way. So, we need to use gizmo location as the basis of our scaling.
+		// On the other hand, we can't directly use CalculateLocalPixelToWorldScale with gizmo origin either because it's
+		//  possible to see the sub gizmo while the gizmo itself is off screen, breaking that calculation. What we want 
+		//  instead is a scale value that is dependent on FOV and distance to gizmo, but independent of view direction. 
+		//  The way we do this is we take the distance to gizmo, and then evaluate the scale at a point that far along the
+		//  view direction (so the actual view direction no longer matters- we're always looking at the scale around the
+		//  center of the screen).
+		// Note that this doesn't fully fix all problems with large component-to-gizmo offsets. Namely, moving the component
+		//  with the gizmo off screen can sometimes make it seem like the component stops moving and starts growing, which 
+		//  looks unnatural if you're not looking at the gizmo and realizing that you are just modifying your angle relative
+		//  to it. However the approach here tends to be the least broken overall.
+		double DistanceToGizmo = FVector3d::Distance(GizmoToWorld.GetLocation(), View.GetViewLocation());
+		FVector3d PointAtWhichToCheckScale = View.GetViewLocation() + DistanceToGizmo * View.GetViewDirection();
+		return GizmoRenderingUtil::CalculateLocalPixelToWorldScale(&View, PointAtWhichToCheckScale);
+	}
+
+	// Apply the settings to get a transform modified by view
+	FTransform GetSubGizmoAdjustedTransform(
+		const UE::GizmoRenderingUtil::ISceneViewInterface& View,
+		const FTransform& CurrentComponentToWorld,
+		const FTransform& GizmoOriginToComponent,
+		const UE::GizmoRenderingUtil::FSubGizmoTransformAdjuster::FSettings& Settings)
+	{
+		if (!Settings.bKeepConstantViewSize && !Settings.bMirrorBasedOnOctant && !Settings.bUseWorldAxesForGizmo)
+		{
+			return CurrentComponentToWorld;
+		}
+		
+		FTransform GizmoToWorld = GizmoOriginToComponent * CurrentComponentToWorld;
+		FTransform ComponentToGizmo = CurrentComponentToWorld.GetRelativeTransform(GizmoToWorld);
+
+		if (Settings.bUseWorldAxesForGizmo)
+		{
+			GizmoToWorld.SetRotation(FQuat::Identity);
+		}
+
+		if (Settings.bMirrorBasedOnOctant)
+		{
+			FVector3d GizmoSpaceDirectionTowardCamera;
+			if (View.IsPerspectiveProjection())
+			{
+				GizmoSpaceDirectionTowardCamera = GizmoToWorld.InverseTransformPosition(View.GetViewLocation());
+			}
+			else
+			{
+				GizmoSpaceDirectionTowardCamera = GizmoToWorld.InverseTransformVector(View.GetViewDirection());
+			}
+
+			EAxis::Type Axes[3] = { EAxis::X, EAxis::Y, EAxis::Z };
+			for (int Dim = 0; Dim < 3; ++Dim)
+			{
+				if (GizmoSpaceDirectionTowardCamera[Dim] < 0)
+				{
+					ComponentToGizmo.Mirror(Axes[Dim], EAxis::None);
+				}
+			}
+		}
+
+		if (Settings.bKeepConstantViewSize)
+		{
+			double ExtraScaling = GetComponentToGizmoScaling(View, GizmoToWorld);
+
+			ComponentToGizmo.MultiplyScale3D(FVector3d(ExtraScaling));
+			ComponentToGizmo.ScaleTranslation(ExtraScaling);
+		}
+
+		return ComponentToGizmo * GizmoToWorld;
+	}
+}
+
+// FSimpleConstantViewScaleAdjuster:
+
+FTransform UE::GizmoRenderingUtil::FSimpleConstantViewScaleAdjuster::GetAdjustedComponentToWorld(
+	const ISceneViewInterface& View, const FTransform& CurrentComponentToWorld)
+{
+	double ExtraScaling = ::GizmoRenderingUtil::CalculateLocalPixelToWorldScale(&View, CurrentComponentToWorld.GetLocation());
+	FTransform AdjustedTransform = CurrentComponentToWorld;
+	AdjustedTransform.MultiplyScale3D(FVector(ExtraScaling));
+
+	return AdjustedTransform;
+}
+
+
+// FSubGizmoTransformAdjuster:
+
+void UE::GizmoRenderingUtil::FSubGizmoTransformAdjuster::SetGizmoOriginTransform(
+	const FTransform& GizmoOriginToComponentIn)
+{
+	GizmoOriginToComponent_GameThread = GizmoOriginToComponentIn;
+
+	// Safely update the render thread side transform.
+	ENQUEUE_RENDER_COMMAND(FViewpointOctantMirrorTransformAdjusterUpdate)(
+		[GizmoOriginToComponentIn, this](FRHICommandListImmediate& RHICmdList)
+	{
+		GizmoOriginToComponent_RenderThread = GizmoOriginToComponentIn;
+	});
+}
+
+TSharedPtr<UE::GizmoRenderingUtil::FSubGizmoTransformAdjuster> UE::GizmoRenderingUtil::FSubGizmoTransformAdjuster::AddTransformAdjuster(
+	UViewAdjustedStaticMeshGizmoComponent* ComponentIn, USceneComponent* ComponentToKeepDistanceConstantTo, bool bMirrorBasedOnOctant)
+{
+	if (!ComponentIn)
+	{
+		return nullptr;
+	}
+
+	TSharedPtr<UE::GizmoRenderingUtil::FSubGizmoTransformAdjuster> TransformAdjuster =
+		MakeShared<UE::GizmoRenderingUtil::FSubGizmoTransformAdjuster>();
+	UE::GizmoRenderingUtil::FSubGizmoTransformAdjuster::FSettings Settings;
+	Settings.bKeepConstantViewSize = true;
+	Settings.bMirrorBasedOnOctant = bMirrorBasedOnOctant;
+	// bUseWorldAxes gets updated automatically as part of UViewAdjustedStaticMeshGizmoComponent::UpdateWorldLocalState
+	TransformAdjuster->SetSettings(Settings);
+	if (ComponentToKeepDistanceConstantTo)
+	{
+		TransformAdjuster->SetGizmoOriginTransform(
+			ComponentToKeepDistanceConstantTo->GetComponentTransform().GetRelativeTransform(ComponentIn->GetComponentToWorld()));
+	}
+	ComponentIn->SetTransformAdjuster(TransformAdjuster);
+
+	return TransformAdjuster;
+}
+
+FTransform UE::GizmoRenderingUtil::FSubGizmoTransformAdjuster::GetAdjustedComponentToWorld(
+	const ISceneViewInterface& View, const FTransform& CurrentComponentToWorld)
+{
+	return ViewBasedTransformAdjustersLocals::GetSubGizmoAdjustedTransform(View, CurrentComponentToWorld,
+		GizmoOriginToComponent_GameThread, Settings);
+}
+
+FTransform UE::GizmoRenderingUtil::FSubGizmoTransformAdjuster::GetAdjustedComponentToWorld_RenderThread(
+	const ISceneViewInterface& View, const FTransform& CurrentComponentToWorld)
+{
+	return ViewBasedTransformAdjustersLocals::GetSubGizmoAdjustedTransform(View, CurrentComponentToWorld,
+		GizmoOriginToComponent_RenderThread, Settings);
+}
+
+
+// FConstantViewRelativeTransformAdjuster:
+
+FTransform UE::GizmoRenderingUtil::FConstantViewRelativeTransformAdjuster::GetAdjustedComponentToWorld(
+	const ISceneViewInterface& View, const FTransform& CurrentComponentToWorld)
+{
+	UE::Geometry::FFrame3d ViewFrameAtComponent(CurrentComponentToWorld.GetLocation(), View.GetViewDirection(), View.GetViewRight(), View.GetViewUp());
+	FTransform ViewRelativeTransformToUse = ViewRelativeTransform;
+
+	if (bKeepConstantViewSize)
+	{
+		// We're assuming that the gizmo origin is at the component location, so we don't need to worry about scaling
+		//  relative to a different point, but we still want our scaling to be view independent for cases where the
+		//  relative transform is big enough that the origin is offscreen while the component is visible.
+		double ExtraScaling = ViewBasedTransformAdjustersLocals::GetComponentToGizmoScaling(View, CurrentComponentToWorld);
+
+		ViewRelativeTransformToUse.MultiplyScale3D(FVector3d(ExtraScaling));
+		ViewRelativeTransformToUse.ScaleTranslation(ExtraScaling);
+	}
+
+	return ViewRelativeTransformToUse * ViewFrameAtComponent.ToFTransform();
+}
\ No newline at end of file
diff --git a/Engine/Source/Runtime/InteractiveToolsFramework/Public/BaseGizmos/GizmoBaseComponent.h b/Engine/Source/Runtime/InteractiveToolsFramework/Public/BaseGizmos/GizmoBaseComponent.h
--- a/Engine/Source/Runtime/InteractiveToolsFramework/Public/BaseGizmos/GizmoBaseComponent.h
+++ b/Engine/Source/Runtime/InteractiveToolsFramework/Public/BaseGizmos/GizmoBaseComponent.h
@@ -5,10 +5,33 @@
 #include "CoreMinimal.h"
 #include "UObject/ObjectMacros.h"
 #include "Components/PrimitiveComponent.h"
+
 #include "GizmoBaseComponent.generated.h"
 
 class UGizmoViewContext;
 
+UINTERFACE(MinimalAPI)
+class UGizmoBaseComponentInterface : public UInterface
+{
+	GENERATED_BODY()
+};
+
+/**
+ * Interface that allows a component to receive various gizmo-specific callbacks while
+ * still inheriting from some class other than UGizmoBaseComponent.
+ */
+class IGizmoBaseComponentInterface
+{
+	GENERATED_BODY()
+public:
+	virtual void UpdateHoverState(bool bHoveringIn) {}
+
+	virtual void UpdateWorldLocalState(bool bWorldIn) {}
+
+	virtual void UpdateInteractingState(bool bInteractingIn) {}
+};
+
+
 /**
  * Base class for simple Components intended to be used as part of 3D Gizmos.
  * Contains common properties and utility functions.
@@ -16,6 +39,7 @@ class UGizmoViewContext;
  */
 UCLASS(ClassGroup = Utility, HideCategories = (Physics, Collision, Mobile), MinimalAPI)
 class UGizmoBaseComponent : public UPrimitiveComponent
+	, public IGizmoBaseComponentInterface
 {
 	GENERATED_BODY()
 
@@ -55,7 +79,7 @@ public:
 
 public:
 	UFUNCTION()
-	void UpdateHoverState(bool bHoveringIn)
+	virtual void UpdateHoverState(bool bHoveringIn) override
 	{
 		if (bHoveringIn != bHovering)
 		{
@@ -64,7 +88,7 @@ public:
 	}
 
 	UFUNCTION()
-	void UpdateWorldLocalState(bool bWorldIn)
+	virtual void UpdateWorldLocalState(bool bWorldIn) override
 	{
 		if (bWorldIn != bWorld)
 		{
diff --git a/Engine/Source/Runtime/InteractiveToolsFramework/Public/BaseGizmos/GizmoRenderingUtil.h b/Engine/Source/Runtime/InteractiveToolsFramework/Public/BaseGizmos/GizmoRenderingUtil.h
--- a/Engine/Source/Runtime/InteractiveToolsFramework/Public/BaseGizmos/GizmoRenderingUtil.h
+++ b/Engine/Source/Runtime/InteractiveToolsFramework/Public/BaseGizmos/GizmoRenderingUtil.h
@@ -6,25 +6,96 @@
 #include "Components/PrimitiveComponent.h"
 #include "SceneView.h"
 
+class IToolsContextRenderAPI;
 class UGizmoViewContext;
+class UInteractiveGizmo;
+class UInteractiveGizmoManager;
+class UInteractiveToolManager;
+class UViewAdjustedStaticMeshGizmoComponent;
 
 /**
  * Utility functions for standard GizmoComponent rendering
  */
-namespace GizmoRenderingUtil
+
+namespace UE::GizmoRenderingUtil
 {
+	// Interface meant to wrap either an FSceneView or UGizmoViewContext so that a user can write one
+	// function to handle either one (for rendering and for hit testing).
+	class ISceneViewInterface
+	{
+	public:
+		virtual const FIntRect& GetUnscaledViewRect() const = 0;
+		virtual FVector4 WorldToScreen(const FVector&) const = 0;
+		virtual FVector GetViewLocation() const = 0;
+		virtual FVector GetViewDirection() const = 0;
+		virtual FVector GetViewRight() const = 0;
+		virtual FVector GetViewUp() const = 0;
+		virtual const FMatrix& GetViewMatrix() const = 0;
+		virtual bool IsPerspectiveProjection() const = 0;
+	};
+
+	// Wrapper around FSceneView to access it through ISceneViewInterface
+	class FSceneViewWrapper : public ISceneViewInterface
+	{
+	public:
+		explicit FSceneViewWrapper(const FSceneView& SceneView)
+		{
+			View = &SceneView;
+		}
+
+		virtual const FIntRect& GetUnscaledViewRect() const { return View->UnscaledViewRect; }
+		virtual FVector4 WorldToScreen(const FVector& VectorIn) const { return View->WorldToScreen(VectorIn); }
+		virtual FVector GetViewLocation() const { return View->ViewLocation; }
+		virtual FVector GetViewDirection() const { return View->GetViewDirection(); }
+		virtual FVector GetViewRight() const { return View->GetViewRight(); }
+		virtual FVector GetViewUp() const { return View->GetViewUp(); }
+		virtual const FMatrix& GetViewMatrix() const { return View->ViewMatrices.GetViewMatrix(); }
+		virtual bool IsPerspectiveProjection() const { return View->IsPerspectiveProjection(); }
+
+	private:
+		const FSceneView* View;
+	};
+
+
+	// Can be used as TranslucencySortPriority to make gizmo elements show up above other translucent objects
+	const int32 GIZMO_TRANSLUCENCY_SORT_PRIORITY = 5000;
 
 	/**
-	 * @return Conversion factor between pixel and world-space coordinates at 3D point Location in View.
-	 * @warning This is a local estimate and is increasingly incorrect as the 3D point gets further from Location
+	 * Gets a custom material suitable to use for gizmo components. The material is drawn on top of opaque geometry with
+	 *  dithering for portions behind opaque materials, but uses the custom depth buffer to properly occlude itself.
+	 * Components using this material should set bRenderCustomDepth to true so they can occlude other gizmo elements. It
+	 *  is also suggested that TranslucencySortPriority be set to something like UE::GizmoRenderingUtil::GIZMO_TRANSLUCENCY_SORT_PRIORITY
+	 *  so that the component is drawn on top of other translucent materials.
+	 * 
+	 * @param Outer Object to set as outer for the material instance. Typically can be left as nullptr to use transient package.
 	 */
-	INTERACTIVETOOLSFRAMEWORK_API float CalculateLocalPixelToWorldScale(
-		const FSceneView* View,
-		const FVector& Location);
-	INTERACTIVETOOLSFRAMEWORK_API float CalculateLocalPixelToWorldScale(
-		const UGizmoViewContext* ViewContext,
-		const FVector& Location);
+	INTERACTIVETOOLSFRAMEWORK_API UMaterialInterface* GetDefaultGizmoComponentMaterial(const FLinearColor& Color, UObject* Outer = nullptr);
 
+	/**
+	 * Helper that creates a component with the default gizmo material and sets up the component-side properties that are needed for
+	 *  it to properly work (translucency sort priority, etc). This may not be necessary if your component is using some other
+	 *  gizmo material that doesn't require component-side flags.
+	 */
+	INTERACTIVETOOLSFRAMEWORK_API UViewAdjustedStaticMeshGizmoComponent* CreateDefaultMaterialGizmoMeshComponent(
+		UStaticMesh* Mesh, UGizmoViewContext* GizmoViewContext, UObject* OwnerComponentOrActor, 
+		const FLinearColor& Color, bool bAddHoverMaterial = true);
+
+	/**
+	 * Overload that takes a gizmo manager to get the view context from that.
+	 */
+	INTERACTIVETOOLSFRAMEWORK_API UViewAdjustedStaticMeshGizmoComponent* CreateDefaultMaterialGizmoMeshComponent(
+		UStaticMesh* Mesh, UInteractiveGizmoManager* GizmoManager, UObject* OwnerComponentOrActor, 
+		const FLinearColor& Color, bool bAddHoverMaterial = true);
+
+	/**
+	 * Gets a red/green/blue color based on the axis (X, Y, or Z).
+	 */
+	INTERACTIVETOOLSFRAMEWORK_API FLinearColor GetDefaultAxisColor(EAxis::Type Axis);
+}
+
+//~ TODO: These should someday be moved to nest under the UE namespace above
+namespace GizmoRenderingUtil
+{
 	/**
 	 * @return Conversion factor between pixel and world-space coordinates at 3D point Location in View.
 	 * @warning This is a local estimate and is increasingly incorrect as the 3D point gets further from Location
@@ -32,6 +103,9 @@ namespace GizmoRenderingUtil
 	INTERACTIVETOOLSFRAMEWORK_API float CalculateLocalPixelToWorldScale(
 		const FSceneView* View,
 		const FVector& Location);
+	INTERACTIVETOOLSFRAMEWORK_API float CalculateLocalPixelToWorldScale(
+		const UE::GizmoRenderingUtil::ISceneViewInterface* View,
+		const FVector& Location);
 
 	/**
 	 * @return Legacy view dependent conversion factor.
diff --git a/Engine/Source/Runtime/InteractiveToolsFramework/Public/BaseGizmos/GizmoViewContext.h b/Engine/Source/Runtime/InteractiveToolsFramework/Public/BaseGizmos/GizmoViewContext.h
--- a/Engine/Source/Runtime/InteractiveToolsFramework/Public/BaseGizmos/GizmoViewContext.h
+++ b/Engine/Source/Runtime/InteractiveToolsFramework/Public/BaseGizmos/GizmoViewContext.h
@@ -4,6 +4,7 @@
 
 #include "CoreMinimal.h"
 
+#include "GizmoRenderingUtil.h"
 #include "SceneView.h"
 
 #include "GizmoViewContext.generated.h"
@@ -15,7 +16,7 @@
  * or UGizmoViewContext, though UGizmoViewContext only keeps the needed data.
  */
 UCLASS(MinimalAPI)
-class UGizmoViewContext : public UObject
+class UGizmoViewContext : public UObject, public UE::GizmoRenderingUtil::ISceneViewInterface
 {
 	GENERATED_BODY()
 public:
@@ -51,15 +52,16 @@ public:
 		ViewLocation = SceneView.ViewLocation;
 	}
 
-	// FSceneView-like functions/properties:
-	FVector GetViewRight() const { return ViewMatrices.GetViewMatrix().GetColumn(0); }
-	FVector GetViewUp() const { return ViewMatrices.GetViewMatrix().GetColumn(1); }
-	FVector GetViewDirection() const { return ViewMatrices.GetViewMatrix().GetColumn(2); }
+	// ISceneViewInterface
+	const FIntRect& GetUnscaledViewRect() const override { return UnscaledViewRect; }
+	FVector GetViewLocation() const override { return ViewLocation; }
+	FVector GetViewRight() const override { return ViewMatrices.GetViewMatrix().GetColumn(0); }
+	FVector GetViewUp() const override { return ViewMatrices.GetViewMatrix().GetColumn(1); }
+	FVector GetViewDirection() const override { return ViewMatrices.GetViewMatrix().GetColumn(2); }
+	const FMatrix& GetViewMatrix() const override { return ViewMatrices.GetViewMatrix(); }
+	bool IsPerspectiveProjection() const override { return bIsPerspectiveProjection; }
 
-	// As a function just for similarity with FSceneView
-	bool IsPerspectiveProjection() const { return bIsPerspectiveProjection; }
-
-	FVector4 WorldToScreen(const FVector& WorldPoint) const
+	FVector4 WorldToScreen(const FVector& WorldPoint) const override
 	{
 		return ViewMatrices.GetViewProjectionMatrix().TransformFVector4(FVector4(WorldPoint, 1));
 	}
diff --git a/Engine/Source/Runtime/InteractiveToolsFramework/Public/BaseGizmos/ViewAdjustedStaticMeshGizmoComponent.h b/Engine/Source/Runtime/InteractiveToolsFramework/Public/BaseGizmos/ViewAdjustedStaticMeshGizmoComponent.h
new file mode 100644
--- /dev/null
+++ b/Engine/Source/Runtime/InteractiveToolsFramework/Public/BaseGizmos/ViewAdjustedStaticMeshGizmoComponent.h
@@ -0,0 +1,92 @@
+// Copyright Epic Games, Inc. All Rights Reserved.
+
+#pragma once
+
+#include "Components/StaticMeshComponent.h"
+#include "BaseGizmos/GizmoBaseComponent.h"
+
+#include "ViewAdjustedStaticMeshGizmoComponent.generated.h"
+
+class UGizmoViewContext;
+namespace UE::GizmoRenderingUtil
+{
+	class IViewBasedTransformAdjuster;
+}
+
+/**
+ * Version of a static mesh component that only takes the dynamic draw path and has the ability to
+ * adjust the transform based on view information.
+ */
+UCLASS(Blueprintable, BlueprintType, ClassGroup = Gizmos, meta = (BlueprintSpawnableComponent))
+class INTERACTIVETOOLSFRAMEWORK_API UViewAdjustedStaticMeshGizmoComponent : public UStaticMeshComponent
+	, public IGizmoBaseComponentInterface
+{
+	GENERATED_BODY()
+public:
+	/**
+	 * The gizmo view context is needed to be able to line trace the component, since its collision data
+	 *  needs updating based on view.
+	 */
+	void SetGizmoViewContext(UGizmoViewContext* GizmoViewContextIn) { GizmoViewContext = GizmoViewContextIn; }
+	
+	void SetTransformAdjuster(TSharedPtr<UE::GizmoRenderingUtil::IViewBasedTransformAdjuster> Adjuster);
+	TSharedPtr<UE::GizmoRenderingUtil::IViewBasedTransformAdjuster> GetTransformAdjuster() { return TransformAdjuster; }
+	
+	/** Helper method that just sets the same material in all slots. Does not include hover override material. */
+	void SetAllMaterials(UMaterialInterface* Material);
+
+	/**
+	 * Sets a material that will override all material slots whenever the component is told that is being
+	 *  hovered (via UpdateHoverState).
+	 */
+	void SetHoverOverrideMaterial(UMaterialInterface* Material);
+
+	UMaterialInterface* GetHoverOverrideMaterial() { return HoverOverrideMaterial; }
+	bool IsBeingHovered() { return bHovered; }
+
+	/** 
+	 * Sets a mesh that is swapped in while the component is being interacted with. This is done by not rendering this 
+	 *  component and making the substitute component visible.
+	 */
+	void SetSubstituteInteractionComponent(UPrimitiveComponent* Component, const FTransform& RelativeTransform = FTransform::Identity);
+	
+	bool IsHiddenByInteraction() { return bInteracted && SubstituteInteractionComponent; }
+
+	// IGizmoBaseComponentInterface
+	virtual void UpdateHoverState(bool bHoveringIn) override;
+	virtual void UpdateWorldLocalState(bool bWorldIn) override;
+	virtual void UpdateInteractingState(bool bInteracting) override;
+
+	// UMeshComponent
+	virtual FMaterialRelevance GetMaterialRelevance(ERHIFeatureLevel::Type InFeatureLevel) const override;
+
+	// UPrimitiveComponent
+	virtual bool LineTraceComponent(FHitResult& OutHit, const FVector Start, const FVector End, const FCollisionQueryParams& Params) override;
+
+	// UActorComponent
+	virtual bool IsHLODRelevant() const override { return false; }
+
+	// UObject
+	virtual bool NeedsLoadForServer() const override { return false; }
+
+protected:
+	// UStaticMeshComponent
+	virtual FPrimitiveSceneProxy* CreateStaticMeshSceneProxy(Nanite::FMaterialAudit& NaniteMaterials, bool bCreateNanite) override;
+
+private:
+		
+	// Needed for proper line traces
+	UPROPERTY()
+	TObjectPtr<UGizmoViewContext> GizmoViewContext = nullptr;
+
+	UPROPERTY()
+	TObjectPtr<UPrimitiveComponent> SubstituteInteractionComponent = nullptr;
+
+	UPROPERTY()
+	TObjectPtr<UMaterialInterface> HoverOverrideMaterial = nullptr;
+
+	TSharedPtr<UE::GizmoRenderingUtil::IViewBasedTransformAdjuster> TransformAdjuster;
+
+	bool bHovered = false;
+	bool bInteracted = false;
+};
diff --git a/Engine/Source/Runtime/InteractiveToolsFramework/Public/BaseGizmos/ViewBasedTransformAdjusters.h b/Engine/Source/Runtime/InteractiveToolsFramework/Public/BaseGizmos/ViewBasedTransformAdjusters.h
new file mode 100644
--- /dev/null
+++ b/Engine/Source/Runtime/InteractiveToolsFramework/Public/BaseGizmos/ViewBasedTransformAdjusters.h
@@ -0,0 +1,167 @@
+// Copyright Epic Games, Inc. All Rights Reserved.
+
+#pragma once
+
+#include "Math/Transform.h" // Compiler will complain with just MathFwd here due to the inline with FTransform return
+#include "Templates/SharedPointer.h"
+
+class USceneComponent;
+class UViewAdjustedStaticMeshGizmoComponent;
+
+namespace UE::GizmoRenderingUtil
+{
+	class ISceneViewInterface;
+
+	/**
+	 * Interface for a helper that can adjust a component's transform based on view information,
+	 *  used by UViewAdjustedStaticMeshGizmoComponent. Adjusters are typically expected to be safely 
+	 *  shareable across game and render threads, so if they are not constant, they are expected to
+	 *  handle updates safely in their implementation.
+	 */
+	class IViewBasedTransformAdjuster
+	{
+	public:
+		/**
+		 * Given the component location and the view information, gives the desired transform of the component.
+		 */
+		virtual FTransform GetAdjustedComponentToWorld(
+			const ISceneViewInterface& View, const FTransform& CurrentComponentToWorld) = 0;
+
+		/**
+		 * If an adjuster is used by a render proxy, it will use this endpoint for getting the transform.
+		 *  This allows an adjuster to be potentially shared between game and render threads. Of course,
+		 *  a component might instead choose to create a new adjuster whenever the render proxy is recreated,
+		 *  and recreate the proxy whenever any relevant parameters change.
+		 */
+		virtual FTransform GetAdjustedComponentToWorld_RenderThread(
+			const ISceneViewInterface& View, const FTransform& CurrentComponentToWorld)
+		{
+			// By default, just route to the other implementation. This is safe to do for any adjusters
+			//  that remain constant for the life of the render proxy.
+			return GetAdjustedComponentToWorld(View, CurrentComponentToWorld);
+		}
+
+		/**
+		 * Not every adjuster will care about the world/local setting, but it is useful for the
+		 *  a component that uses it to be able to blindly pass this information down without
+		 *  caring about the type of adjuster.
+		 */
+		virtual void UpdateWorldLocalState(bool bWorldIn) {}
+	};
+
+	/**
+	 * Adjuster that maintains same view size but otherwise keeps the component transform. Note
+	 *  that this will look wrong if the component is not at the gizmo origin, because the distance
+	 *  relative to gizmo origin won't be scaled. Use FGizmoSubComponentTransformAdjuster if that is needed.
+	 */
+	class FSimpleConstantViewScaleAdjuster : public IViewBasedTransformAdjuster
+	{
+	public:
+		// IGizmoViewTransformAdjuster
+		INTERACTIVETOOLSFRAMEWORK_API virtual FTransform GetAdjustedComponentToWorld(
+			const ISceneViewInterface& View,
+			const FTransform& CurrentComponentToWorld) override;
+	private:
+	};
+
+	/**
+	 * An adjuster that can do various transformations common for sub gizmos, which are based off
+	 *  of the parent gizmo transform (in addition to the actual component transform).
+	 */
+	class FSubGizmoTransformAdjuster : public IViewBasedTransformAdjuster
+	{
+	public:
+		struct FSettings
+		{
+			// Some compilers (clang, at least) have issues with using a nested class with default initializers
+			// and a default constructor as a default argument in a function, hence this constructor.
+			FSettings() {};
+
+			// Keeps view size and offset relative to parent gizmo constant
+			bool bKeepConstantViewSize = true;
+			// Mirrors the component depending on which octant of the parent gizmo the view is located
+			bool bMirrorBasedOnOctant = true;
+			// Considers the parent gizmo transform to be unrotated, and applies the relative component transform on top of that.
+			//  This is frequently not necessary to set because code upstream will typically update it every tick anyway
+			//  via UViewAdjustedStaticMeshGizmoComponent's UpdateWorldLocalState implementation.
+			bool bUseWorldAxesForGizmo = false;
+		};
+
+		FSubGizmoTransformAdjuster(
+			const FTransform& GizmoOriginToComponent = FTransform::Identity, const FSettings& SettingsIn = FSettings())
+			: GizmoOriginToComponent_GameThread(GizmoOriginToComponent)
+			, GizmoOriginToComponent_RenderThread(GizmoOriginToComponent)
+			, Settings(SettingsIn)
+		{}
+
+		/**
+		 * Update the transform of the gizmo origin relative to this component. This is safe to do even if the
+		 *  adjuster is used by the render thread (the render-side update gets queued properly).
+		 */
+		INTERACTIVETOOLSFRAMEWORK_API void SetGizmoOriginTransform(const FTransform& GizmoOriginToComponent);
+
+		void SetSettings(const FSettings& SettingsIn)
+		{
+			Settings = SettingsIn;
+		}
+
+		/**
+		 * Static helper method to create and add this adjuster to a gizmo component for the common case of keeping a 
+		 *  constant size gizmo.
+		 * 
+		 * @param bMirrorBasedOnOctant If true, mirrors the component around the ComponentToKeepDistanceConstantTo
+		 *  depending on which octant of that component the view is located in.
+		 */
+		INTERACTIVETOOLSFRAMEWORK_API static TSharedPtr<FSubGizmoTransformAdjuster> AddTransformAdjuster(
+			UViewAdjustedStaticMeshGizmoComponent* ComponentIn, USceneComponent* ComponentToKeepDistanceConstantTo, 
+			bool bMirrorBasedOnOctant);
+
+
+		// IGizmoViewTransformAdjuster
+		INTERACTIVETOOLSFRAMEWORK_API virtual FTransform GetAdjustedComponentToWorld(
+			const GizmoRenderingUtil::ISceneViewInterface& View,
+			const FTransform& CurrentComponentToWorld) override;
+
+		INTERACTIVETOOLSFRAMEWORK_API virtual FTransform GetAdjustedComponentToWorld_RenderThread(
+			const ISceneViewInterface& View, const FTransform& CurrentComponentToWorld) override;
+		
+		virtual void UpdateWorldLocalState(bool bWorldIn) override
+		{
+			Settings.bUseWorldAxesForGizmo = bWorldIn;
+		}
+	private:
+		FTransform GizmoOriginToComponent_GameThread = FTransform::Identity;
+		FTransform GizmoOriginToComponent_RenderThread = FTransform::Identity;
+		// We might decide to have a render thread version of the settings as well, but it's unclear whether it's
+		//  worth it, since it seems like a brief inconsistency during the update is unlikely to cause much of a
+		//  problem: a user is unlikely to be changing multiple bools at once, and reading an outdated bool seems
+		//  fine.
+		FSettings Settings;
+	};
+
+	/**
+	 * Adjuster that takes the the view frame at the component world location, and applies a constant
+	 *  relative transform to that. This can be used to create billboard-like gizmo components that keep
+	 *  some orientation to the camera.
+	 * This doesn't scale the offset relative to gizmo center, but the simplest option is typically to
+	 *  keep the component at gizmo origin and adjust the view relative transform appropriately.
+	 */
+	//~ TODO: May someday want a non-const version of this where the relative transform is editable, for sliders.
+	class FConstantViewRelativeTransformAdjuster : public IViewBasedTransformAdjuster
+	{
+	public:
+		FConstantViewRelativeTransformAdjuster(const FTransform& ViewRelativeTransformIn, bool bConstantSize = true)
+			: ViewRelativeTransform(ViewRelativeTransformIn)
+			, bKeepConstantViewSize(bConstantSize)
+		{
+		}
+
+		// IGizmoViewTransformAdjuster
+		INTERACTIVETOOLSFRAMEWORK_API virtual FTransform GetAdjustedComponentToWorld(
+			const GizmoRenderingUtil::ISceneViewInterface& View,
+			const FTransform& CurrentComponentToWorld) override;
+	private:
+		FTransform ViewRelativeTransform = FTransform::Identity;
+		bool bKeepConstantViewSize = true;
+	};
+}
\ No newline at end of file
