InteractiveToolsFramework: Add initialization helper methods to various sub gizmos so that they can be initialized independently of a TRS gizmo. Use these helpers in the TRS gizmo. Also make the sub gizmos update the existing hit target interaction callback.

diff --git a/Engine/Source/Runtime/InteractiveToolsFramework/Private/BaseGizmos/AxisAngleGizmo.cpp b/Engine/Source/Runtime/InteractiveToolsFramework/Private/BaseGizmos/AxisAngleGizmo.cpp
--- a/Engine/Source/Runtime/InteractiveToolsFramework/Private/BaseGizmos/AxisAngleGizmo.cpp
+++ b/Engine/Source/Runtime/InteractiveToolsFramework/Private/BaseGizmos/AxisAngleGizmo.cpp
@@ -5,6 +5,8 @@
 #include "BaseBehaviors/ClickDragBehavior.h"
 #include "BaseBehaviors/MouseHoverBehavior.h"
 #include "BaseGizmos/GizmoMath.h"
+#include "BaseGizmos/GizmoPrivateUtil.h" // SetCommonSubGizmoProperties
+#include "BaseGizmos/TransformSubGizmoUtil.h" // FTransformSubGizmoCommonParams
 
 #include UE_INLINE_GENERATED_CPP_BY_NAME(AxisAngleGizmo)
 
@@ -41,6 +43,32 @@ void UAxisAngleGizmo::Setup()
 	bInInteraction = false;
 }
 
+bool UAxisAngleGizmo::InitializeAsRotateGizmo(
+	const UE::GizmoUtil::FTransformSubGizmoCommonParams& Params, 
+	UE::GizmoUtil::FTransformSubGizmoSharedState* SharedState)
+{
+	if (!Params.Component
+		|| !Params.TransformProxy
+		|| Params.Axis == EAxis::None)
+	{
+		return false;
+	}
+
+	UGizmoScaledAndUnscaledTransformSources* TransformSource;
+	if (!UE::GizmoUtil::SetCommonSubGizmoProperties(this, Params, SharedState, TransformSource))
+	{
+		return false;
+	}
+
+	UObject* Owner = Params.OuterForSubobjects ? Params.OuterForSubobjects : GetTransientPackage();
+
+	// Parameter source maps axis-parameter-change to translation of TransformSource's transform
+	UGizmoAxisRotationParameterSource* CastAngleSource = UGizmoAxisRotationParameterSource::Construct(
+		AxisSource.GetInterface(), TransformSource, Owner);
+	AngleSource = CastAngleSource;
+
+	return true;
+}
 
 void UAxisAngleGizmo::OnUpdateModifierState(int ModifierID, bool bIsOn)
 {
@@ -106,6 +134,10 @@ void UAxisAngleGizmo::OnClickPress(const FInputDeviceRay& PressPos)
 	{
 		StateTarget->BeginUpdate();
 	}
+	if (ensure(HitTarget))
+	{
+		HitTarget->UpdateInteractingState(bInInteraction);
+	}
 }
 
 
@@ -165,6 +197,11 @@ void UAxisAngleGizmo::OnClickRelease(const FInputDeviceRay& ReleasePos)
 		StateTarget->EndUpdate();
 	}
 	bInInteraction = false;
+
+	if (ensure(HitTarget))
+	{
+		HitTarget->UpdateInteractingState(bInInteraction);
+	}
 }
 
 
@@ -178,6 +215,11 @@ void UAxisAngleGizmo::OnTerminateDragSequence()
 		StateTarget->EndUpdate();
 	}
 	bInInteraction = false;
+
+	if (ensure(HitTarget))
+	{
+		HitTarget->UpdateInteractingState(bInInteraction);
+	}
 }
 
 
diff --git a/Engine/Source/Runtime/InteractiveToolsFramework/Private/BaseGizmos/AxisPositionGizmo.cpp b/Engine/Source/Runtime/InteractiveToolsFramework/Private/BaseGizmos/AxisPositionGizmo.cpp
--- a/Engine/Source/Runtime/InteractiveToolsFramework/Private/BaseGizmos/AxisPositionGizmo.cpp
+++ b/Engine/Source/Runtime/InteractiveToolsFramework/Private/BaseGizmos/AxisPositionGizmo.cpp
@@ -1,10 +1,14 @@
 // Copyright Epic Games, Inc. All Rights Reserved.
 
 #include "BaseGizmos/AxisPositionGizmo.h"
+
 #include "InteractiveGizmoManager.h"
 #include "BaseBehaviors/ClickDragBehavior.h"
 #include "BaseBehaviors/MouseHoverBehavior.h"
 #include "BaseGizmos/GizmoMath.h"
+#include "BaseGizmos/GizmoPrivateUtil.h" // SetCommonSubGizmoProperties
+#include "BaseGizmos/TransformSubGizmoUtil.h" // FTransformSubGizmoCommonParams
+#include "Components/PrimitiveComponent.h"
 
 #include UE_INLINE_GENERATED_CPP_BY_NAME(AxisPositionGizmo)
 
@@ -41,6 +45,88 @@ void UAxisPositionGizmo::Setup()
 	bInInteraction = false;
 }
 
+bool UAxisPositionGizmo::InitializeAsTranslateGizmo(
+	const UE::GizmoUtil::FTransformSubGizmoCommonParams& Params,
+	UE::GizmoUtil::FTransformSubGizmoSharedState* SharedState)
+{
+	if (!Params.Component
+		|| !Params.TransformProxy
+		|| Params.Axis == EAxis::None)
+	{
+		return false;
+	}
+
+	UGizmoScaledAndUnscaledTransformSources* TransformSource;
+	if (!UE::GizmoUtil::SetCommonSubGizmoProperties(this, Params, SharedState, TransformSource))
+	{
+		return false;
+	}
+
+	UObject* Owner = Params.OuterForSubobjects ? Params.OuterForSubobjects : GetTransientPackage();
+
+	// Parameter source maps axis-parameter-change to translation of TransformSource's transform
+	ParameterSource = UGizmoAxisTranslationParameterSource::Construct(
+		AxisSource.GetInterface(), TransformSource, Owner);
+
+	return true;
+}
+
+bool UAxisPositionGizmo::InitializeAsScaleGizmo(
+	const UE::GizmoUtil::FTransformSubGizmoCommonParams& Params,
+	bool bDisallowNegativeScaling, UE::GizmoUtil::FTransformSubGizmoSharedState* SharedState)
+{
+	if (!Params.Component
+		|| !Params.TransformProxy
+		|| Params.Axis == EAxis::None)
+	{
+		return false;
+	}
+
+	int AxisIndex = Params.GetClampedAxisIndex();
+
+	UGizmoScaledAndUnscaledTransformSources* TransformSource;
+	if (!UE::GizmoUtil::SetCommonSubGizmoProperties(this, Params, SharedState, TransformSource))
+	{
+		return false;
+	}
+
+	UObject* Owner = Params.OuterForSubobjects ? Params.OuterForSubobjects : GetTransientPackage();
+	bEnableSignedAxis = true;
+
+	// Although the normal axis source gets used for detecting interactions, the parameter application has
+	// to happen along unrotated axes because the scaling gets applied before rotation. In other words if we
+	// tried to apply scaling measured along a rotated vector, we would end up incorrectly scaling along
+	// multiple axes.
+	UGizmoComponentAxisSource* UnitCardinalAxisSource = nullptr;
+	// See if we already have it in our shared state
+	if (SharedState && SharedState->UnitCardinalAxisSources[AxisIndex])
+	{
+		UnitCardinalAxisSource = SharedState->UnitCardinalAxisSources[AxisIndex];
+	}
+	else
+	{
+		// Create new and add to shared state.
+		USceneComponent* RootComponent = Params.Component->GetOwner()->GetRootComponent();
+		UGizmoComponentAxisSource* CastAxisSource = UGizmoComponentAxisSource::Construct(RootComponent, AxisIndex,
+			// bUseLocalAxes, not important because we're going to be updating this value every tick
+			true,
+			Owner);
+		UnitCardinalAxisSource = CastAxisSource;
+		if (SharedState)
+		{
+			SharedState->UnitCardinalAxisSources[AxisIndex] = UnitCardinalAxisSource;
+		}
+	}
+
+	// Parameter source maps axis-parameter-change to scale of TransformSource's transform
+	UGizmoAxisScaleParameterSource* CastParameterSource = UGizmoAxisScaleParameterSource::Construct(
+		UnitCardinalAxisSource, TransformSource, Owner);
+	ParameterSource = CastParameterSource;
+	CastParameterSource->bClampToZero = bDisallowNegativeScaling;
+
+	return true;
+}
+
 FInputRayHit UAxisPositionGizmo::CanBeginClickDragSequence(const FInputDeviceRay& PressPos)
 {
 	FInputRayHit GizmoHit;
@@ -88,15 +174,14 @@ void UAxisPositionGizmo::OnClickPress(const FInputDeviceRay& PressPos)
 
 	bInInteraction = true;
 
-	if (HitTarget)
-	{
-		HitTarget->UpdateInteractingState(bInInteraction);
-	}
-
 	if (StateTarget)
 	{
 		StateTarget->BeginUpdate();
 	}
+	if (ensure(HitTarget))
+	{
+		HitTarget->UpdateInteractingState(bInInteraction);
+	}
 }
 
 void UAxisPositionGizmo::OnClickDrag(const FInputDeviceRay& DragPos)
@@ -141,7 +226,7 @@ void UAxisPositionGizmo::OnClickRelease(const FInputDeviceRay& ReleasePos)
 	}
 	bInInteraction = false;
 
-	if (HitTarget)
+	if (ensure(HitTarget))
 	{
 		HitTarget->UpdateInteractingState(bInInteraction);
 	}
@@ -158,8 +243,7 @@ void UAxisPositionGizmo::OnTerminateDragSequence()
 		StateTarget->EndUpdate();
 	}
 	bInInteraction = false;
-
-	if (HitTarget)
+	if (ensure(HitTarget))
 	{
 		HitTarget->UpdateInteractingState(bInInteraction);
 	}
diff --git a/Engine/Source/Runtime/InteractiveToolsFramework/Private/BaseGizmos/CombinedTransformGizmo.cpp b/Engine/Source/Runtime/InteractiveToolsFramework/Private/BaseGizmos/CombinedTransformGizmo.cpp
--- a/Engine/Source/Runtime/InteractiveToolsFramework/Private/BaseGizmos/CombinedTransformGizmo.cpp
+++ b/Engine/Source/Runtime/InteractiveToolsFramework/Private/BaseGizmos/CombinedTransformGizmo.cpp
@@ -13,9 +13,13 @@
 #include "BaseGizmos/GizmoCircleComponent.h"
 #include "BaseGizmos/GizmoBoxComponent.h"
 #include "BaseGizmos/GizmoLineHandleComponent.h"
+#include "BaseGizmos/GizmoPrivateUtil.h" // ToAxis
 #include "BaseGizmos/GizmoViewContext.h"
+#include "BaseGizmos/TransformSubGizmoUtil.h" // FTransformSubGizmoCommonParams, FTransformSubGizmoSharedState
 
+#include "Quaternion.h"
 #include "MathUtil.h"
+#include "MatrixTypes.h"
 #include "VectorUtil.h"
 
 // need this to implement hover
@@ -79,6 +83,105 @@ ACombinedTransformGizmoActor::ACombinedTransformGizmoActor()
 
 
 
+
+bool ACombinedTransformGizmoActor::ReplaceSubGizmoComponent(ETransformGizmoSubElements Element,
+	UPrimitiveComponent* NewComponent, const FTransform& SubGizmoToGizmo, UPrimitiveComponent** ReplacedComponentOut)
+{
+	// We allow a null NewComponent (which equates to element removal), but if we do have a component,
+	//  it should have this actor in its outer chain. It might be possible to loosen that restriction,
+	//  but it's likely that something is wrong in this case.
+	if (NewComponent && !ensure(NewComponent->GetOwner() == this))
+	{
+		return false;
+	}
+
+	auto ReplaceComponent = [this, &SubGizmoToGizmo, NewComponent, ReplacedComponentOut](TObjectPtr<UPrimitiveComponent>& ComponentToReplace)
+	{
+		if (ComponentToReplace)
+		{
+			ComponentToReplace->DestroyComponent();
+		}
+		if (ReplacedComponentOut)
+		{
+			*ReplacedComponentOut = ComponentToReplace;
+		}
+
+		ComponentToReplace = NewComponent;
+
+		if (NewComponent)
+		{
+			AddInstanceComponent(NewComponent);
+			NewComponent->AttachToComponent(GetRootComponent(), FAttachmentTransformRules::KeepRelativeTransform);
+			NewComponent->SetRelativeTransform(SubGizmoToGizmo);
+			if (!NewComponent->IsRegistered())
+			{
+				NewComponent->RegisterComponent();
+			}
+		}
+	};
+
+	switch (Element)
+	{
+	case ETransformGizmoSubElements::TranslateAxisX:
+		ReplaceComponent(TranslateX);
+		break;
+	case ETransformGizmoSubElements::TranslateAxisY:
+		ReplaceComponent(TranslateY);
+		break;
+	case ETransformGizmoSubElements::TranslateAxisZ:
+		ReplaceComponent(TranslateZ);
+		break;
+	case ETransformGizmoSubElements::TranslatePlaneXY:
+		ReplaceComponent(TranslateXY);
+		break;
+	case ETransformGizmoSubElements::TranslatePlaneXZ:
+		ReplaceComponent(TranslateXZ);
+		break;
+	case ETransformGizmoSubElements::TranslatePlaneYZ:
+		ReplaceComponent(TranslateYZ);
+		break;
+	case ETransformGizmoSubElements::RotateAxisX:
+		ReplaceComponent(RotateX);
+		break;
+	case ETransformGizmoSubElements::RotateAxisY:
+		ReplaceComponent(RotateY);
+		break;
+	case ETransformGizmoSubElements::RotateAxisZ:
+		ReplaceComponent(RotateZ);
+		break;
+	case ETransformGizmoSubElements::ScaleAxisX:
+		ReplaceComponent(AxisScaleX);
+		break;
+	case ETransformGizmoSubElements::ScaleAxisY:
+		ReplaceComponent(AxisScaleY);
+		break;
+	case ETransformGizmoSubElements::ScaleAxisZ:
+		ReplaceComponent(AxisScaleZ);
+		break;
+	case ETransformGizmoSubElements::ScalePlaneXY:
+		ReplaceComponent(PlaneScaleXY);
+		break;
+	case ETransformGizmoSubElements::ScalePlaneXZ:
+		ReplaceComponent(PlaneScaleXZ);
+		break;
+	case ETransformGizmoSubElements::ScalePlaneYZ:
+		ReplaceComponent(PlaneScaleYZ);
+		break;
+	case ETransformGizmoSubElements::ScaleUniform:
+		ReplaceComponent(UniformScale);
+		break;
+	// We use the RotateAllAxes identifier for replacing the rotation sphere
+	case ETransformGizmoSubElements::RotateAllAxes:
+		ReplaceComponent(RotationSphere);
+		break;
+	default:
+		UE_LOG(LogGeometry, Warning, TEXT("UCombinedTransformGizmo::SetSubGizmoComponent currently only supports a "
+			"single sub gizmo element at a time."));
+		return false;
+	}
+	return true;
+}
+
 ACombinedTransformGizmoActor* ACombinedTransformGizmoActor::ConstructDefault3AxisGizmo(UWorld* World, UGizmoViewContext* GizmoViewContext)
 {
 	return ConstructCustom3AxisGizmo(World, GizmoViewContext,
@@ -359,6 +462,164 @@ void UCombinedTransformGizmo::SetUpdateCoordSystemFunction(TFunction<void(UPrimi
 	UpdateCoordSystemFunction = CoordSysFunction;
 }
 
+bool UCombinedTransformGizmo::SetSubGizmoComponent(ETransformGizmoSubElements Element, 
+	UPrimitiveComponent* NewComponent, const FTransform& SubGizmoToGizmo)
+{
+	using namespace UE::Geometry;
+	using namespace CombinedTransformGizmoLocals;
+
+	if (!GizmoActor)
+	{
+		return false;
+	}
+
+	EAxis::Type Axis = UE::GizmoUtil::ToAxis(Element);
+
+	UPrimitiveComponent* ReplacedComponent = nullptr;
+	if (!GizmoActor->ReplaceSubGizmoComponent(Element, NewComponent, SubGizmoToGizmo, &ReplacedComponent))
+	{
+		return false;
+	}
+
+	if (!ActiveTarget)
+	{
+		// If the target is not set yet, then we're done for now. The rest of the setup
+		// should end up being done correctly once SetActiveTarget is called.
+		return true;
+	}
+
+	// If we got here, we'll need to do some more work to initialize or reinitialize our gizmo
+
+	// Look for the existing gizmo through our gizmo info arrays
+	if (ReplacedComponent)
+	{
+		ActiveComponents.RemoveSwap(ReplacedComponent);
+
+		TArray<FSubGizmoInfo>* ArrayToSearch = nullptr;
+		if ((Element & (ETransformGizmoSubElements::TranslateAllAxes | ETransformGizmoSubElements::TranslateAllPlanes))
+			!= ETransformGizmoSubElements::None)
+		{
+			ArrayToSearch = &TranslationSubGizmos;
+		}
+		else if ((Element & ETransformGizmoSubElements::RotateAllAxes) != ETransformGizmoSubElements::None)
+		{
+			ArrayToSearch = &RotationSubGizmos;
+		}
+		else if ((Element & (ETransformGizmoSubElements::ScaleAllAxes | ETransformGizmoSubElements::ScaleAllPlanes))
+			!= ETransformGizmoSubElements::None)
+		{
+			ArrayToSearch = &NonUniformScaleSubGizmos;
+		}
+		else if ((Element & ETransformGizmoSubElements::ScaleUniform) != ETransformGizmoSubElements::None)
+		{
+			ArrayToSearch = &UniformScaleSubGizmos;
+		}
+
+		int32 GizmoInfoIndex = -1;
+		FSubGizmoInfo* ExistingGizmoInfo = nullptr;
+		if (ensure(ArrayToSearch))
+		{
+			GizmoInfoIndex = ArrayToSearch->IndexOfByPredicate(
+				[ReplacedComponent](const FSubGizmoInfo& GizmoInfo) { return GizmoInfo.Component == ReplacedComponent; });
+			if (GizmoInfoIndex >= 0)
+			{
+				ExistingGizmoInfo = &(*ArrayToSearch)[GizmoInfoIndex];
+			}
+		}
+
+		if (ensure(ExistingGizmoInfo))
+		{
+			// We could call InitializeAs... on an existing gizmo to swap the component, but then we also need to set
+			// our constraint functions, etc. It seems cleaner code-wise to just destroy this gizmo and create a new one
+			// to make sure everything is updated. We just have to make sure we do the removal thoroughly.
+			if (UInteractiveGizmo* ExistingGizmo = ExistingGizmoInfo->Gizmo.Get())
+			{
+				GetGizmoManager()->DestroyGizmo(ExistingGizmo);
+				ActiveGizmos.Remove(ExistingGizmo);
+			}
+			ArrayToSearch->RemoveAtSwap(GizmoInfoIndex);
+		}
+	}
+
+	if (!NewComponent)
+	{
+		// If we're replacing with a nullptr, then we just wanted to remove that component.
+		// No need to add a gizmo back.
+		return true;
+	}
+
+	UE::GizmoUtil::FTransformSubGizmoCommonParams Params;
+	Params.TransformProxy = ActiveTarget;
+	Params.Axis = Axis;
+	Params.Component = NewComponent;
+	Params.TransactionProvider = TransactionProviderAtLastSetActiveTarget;
+	Params.bManipulatesRootComponent = true;
+
+	// The shared data struct should have been created during SetActiveTarget
+	if (!ensure(SubGizmoSharedState.IsValid()))
+	{
+		SubGizmoSharedState = MakeUnique<FTransformSubGizmoSharedState>();
+	}
+
+	switch (Element)
+	{
+	case ETransformGizmoSubElements::TranslateAxisX:
+	case ETransformGizmoSubElements::TranslateAxisY:
+	case ETransformGizmoSubElements::TranslateAxisZ:
+	{
+		AddAxisTranslationGizmo(Params, *SubGizmoSharedState);
+		break;
+	}
+	case ETransformGizmoSubElements::TranslatePlaneXY:
+	case ETransformGizmoSubElements::TranslatePlaneXZ:
+	case ETransformGizmoSubElements::TranslatePlaneYZ:
+	{
+		AddPlaneTranslationGizmo(Params, *SubGizmoSharedState);
+		break;
+	}
+	case ETransformGizmoSubElements::RotateAxisX:
+	case ETransformGizmoSubElements::RotateAxisY:
+	case ETransformGizmoSubElements::RotateAxisZ:
+	{
+		AddAxisRotationGizmo(Params, *SubGizmoSharedState);
+		break;
+	}
+	case ETransformGizmoSubElements::ScaleAxisX:
+	case ETransformGizmoSubElements::ScaleAxisY:
+	case ETransformGizmoSubElements::ScaleAxisZ:
+	{
+		AddAxisScaleGizmo(Params, *SubGizmoSharedState);
+		break;
+	}
+	case ETransformGizmoSubElements::ScalePlaneXY:
+	case ETransformGizmoSubElements::ScalePlaneXZ:
+	case ETransformGizmoSubElements::ScalePlaneYZ:
+	{
+		AddPlaneScaleGizmo(Params, *SubGizmoSharedState);
+		break;
+	}
+	case ETransformGizmoSubElements::ScaleUniform:
+	{
+		AddUniformScaleGizmo(Params, *SubGizmoSharedState);
+		break;
+	}
+	case ETransformGizmoSubElements::RotateAllAxes:
+	{
+		// no gizmo for the sphere currently
+		if (ensure(GizmoActor->RotationSphere == NewComponent))
+		{
+			ActiveComponents.Add(GizmoActor->RotationSphere);
+			RotationSubGizmos.Add(FSubGizmoInfo{ GizmoActor->RotationSphere, nullptr });
+		}
+		break;
+	}
+	default:
+		return ensure(false);
+	}
+
+	return true;
+}
+
 void UCombinedTransformGizmo::SetWorldAlignmentFunctions(
 	TUniqueFunction<bool()>&& ShouldAlignTranslationIn,
 	TUniqueFunction<bool(const FRay&, FVector&)>&& TranslationAlignmentRayCasterIn)
@@ -451,21 +712,27 @@ void UCombinedTransformGizmo::Setup()
 {
 	UInteractiveGizmo::Setup();
 
-	UpdateHoverFunction = [](UPrimitiveComponent* Component, bool bHovering)
+	if (!UpdateHoverFunction)
 	{
-		if (Cast<UGizmoBaseComponent>(Component) != nullptr)
+		UpdateHoverFunction = [](UPrimitiveComponent* Component, bool bHovering)
 		{
-			Cast<UGizmoBaseComponent>(Component)->UpdateHoverState(bHovering);
-		}
-	};
+			if (IGizmoBaseComponentInterface* CastComponent = Cast<IGizmoBaseComponentInterface>(Component))
+			{
+				CastComponent->UpdateHoverState(bHovering);
+			}
+		};
+	}
 
-	UpdateCoordSystemFunction = [](UPrimitiveComponent* Component, EToolContextCoordinateSystem CoordSystem)
+	if (!UpdateCoordSystemFunction)
 	{
-		if (Cast<UGizmoBaseComponent>(Component) != nullptr)
+		UpdateCoordSystemFunction = [](UPrimitiveComponent* Component, EToolContextCoordinateSystem CoordSystem)
 		{
-			Cast<UGizmoBaseComponent>(Component)->UpdateWorldLocalState(CoordSystem == EToolContextCoordinateSystem::World);
-		}
-	};
+			if (IGizmoBaseComponentInterface* CastComponent = Cast<IGizmoBaseComponentInterface>(Component))
+			{
+				CastComponent->UpdateWorldLocalState(CoordSystem == EToolContextCoordinateSystem::World);
+			}
+		};
+	}
 
 	GizmoActor = GizmoActorBuilder->CreateNewGizmoActor(World);
 }
@@ -572,6 +839,7 @@ void UCombinedTransformGizmo::SetActiveTarget(UTransformProxy* Target, IToolCont
 	}
 
 	ActiveTarget = Target;
+	TransactionProviderAtLastSetActiveTarget = TransactionProvider;
 
 	// move gizmo to target location
 	USceneComponent* GizmoComponent = GizmoActor->GetRootComponent();
@@ -581,84 +849,43 @@ void UCombinedTransformGizmo::SetActiveTarget(UTransformProxy* Target, IToolCont
 	GizmoTransform.SetScale3D(FVector(1, 1, 1));
 	GizmoComponent->SetWorldTransform(GizmoTransform);
 
-	UGizmoScaledAndUnscaledTransformSources* TransformSource = UGizmoScaledAndUnscaledTransformSources::Construct(
-		UGizmoTransformProxyTransformSource::Construct(ActiveTarget, this), 
-		GizmoComponent, this);
-
-	// This state target emits an explicit FChange that moves the GizmoActor root component during undo/redo.
-	// It also opens/closes the Transaction that saves/restores the target object locations.
-	if (TransactionProvider == nullptr)
-	{
-		TransactionProvider = GetGizmoManager();
-	}
-	StateTarget = UGizmoTransformChangeStateTarget::Construct(GizmoComponent,
-		LOCTEXT("UCombinedTransformGizmoTransaction", "Transform"), TransactionProvider, this);
-	StateTarget->DependentChangeSources.Add(MakeUnique<FTransformProxyChangeSource>(Target));
-
-	CameraAxisSource = NewObject<UGizmoConstantFrameAxisSource>(this);
+	UE::GizmoUtil::FTransformSubGizmoCommonParams Params;
+	Params.TransformProxy = ActiveTarget;
+	Params.TransactionProvider = TransactionProvider;
+	Params.bManipulatesRootComponent = true;
 
-	// root component provides local X/Y/Z axis, identified by AxisIndex
-	AxisXSource = UGizmoComponentAxisSource::Construct(GizmoComponent, 0, true, this);
-	AxisYSource = UGizmoComponentAxisSource::Construct(GizmoComponent, 1, true, this);
-	AxisZSource = UGizmoComponentAxisSource::Construct(GizmoComponent, 2, true, this);
+	SubGizmoSharedState = MakeUnique<UE::GizmoUtil::FTransformSubGizmoSharedState>();
 
-	// todo should we hold onto these?
-	if (GizmoActor->TranslateX != nullptr)
+	EAxis::Type Axes[3] = { EAxis::X, EAxis::Y, EAxis::Z };
+	UPrimitiveComponent* TranslateAxisComponents[3]{ GizmoActor->TranslateX, GizmoActor->TranslateY, GizmoActor->TranslateZ };
+	for (int AxisIndex = 0; AxisIndex < 3; ++AxisIndex)
 	{
-		UInteractiveGizmo* NewGizmo = AddAxisTranslationGizmo(GizmoActor->TranslateX, GizmoComponent, AxisXSource, TransformSource, StateTarget, 0);
-		ActiveComponents.Add(GizmoActor->TranslateX);
-		TranslationSubGizmos.Add( FSubGizmoInfo{ GizmoActor->TranslateX, NewGizmo } );
-	}
-	if (GizmoActor->TranslateY != nullptr)
-	{
-		UInteractiveGizmo* NewGizmo = AddAxisTranslationGizmo(GizmoActor->TranslateY, GizmoComponent, AxisYSource, TransformSource, StateTarget, 1);
-		ActiveComponents.Add(GizmoActor->TranslateY);
-		TranslationSubGizmos.Add(FSubGizmoInfo{ GizmoActor->TranslateY, NewGizmo });
-	}
-	if (GizmoActor->TranslateZ != nullptr)
-	{
-		UInteractiveGizmo* NewGizmo = AddAxisTranslationGizmo(GizmoActor->TranslateZ, GizmoComponent, AxisZSource, TransformSource, StateTarget, 2);
-		ActiveComponents.Add(GizmoActor->TranslateZ);
-		TranslationSubGizmos.Add(FSubGizmoInfo{ GizmoActor->TranslateZ, NewGizmo });
-	}
-
-
-	if (GizmoActor->TranslateYZ != nullptr)
-	{
-		UInteractiveGizmo* NewGizmo = AddPlaneTranslationGizmo(GizmoActor->TranslateYZ, GizmoComponent, AxisXSource, TransformSource, StateTarget, 1, 2);
-		ActiveComponents.Add(GizmoActor->TranslateYZ);
-		TranslationSubGizmos.Add(FSubGizmoInfo{ GizmoActor->TranslateYZ, NewGizmo });
-	}
-	if (GizmoActor->TranslateXZ != nullptr)
-	{
-		UInteractiveGizmo* NewGizmo = AddPlaneTranslationGizmo(GizmoActor->TranslateXZ, GizmoComponent, AxisYSource, TransformSource, StateTarget, 2, 0);	// flip here corresponds to UGizmoComponentAxisSource::GetTangentVectors()
-		ActiveComponents.Add(GizmoActor->TranslateXZ);
-		TranslationSubGizmos.Add(FSubGizmoInfo{ GizmoActor->TranslateXZ, NewGizmo });
-	}
-	if (GizmoActor->TranslateXY != nullptr)
-	{
-		UInteractiveGizmo* NewGizmo = AddPlaneTranslationGizmo(GizmoActor->TranslateXY, GizmoComponent, AxisZSource, TransformSource, StateTarget, 0, 1);
-		ActiveComponents.Add(GizmoActor->TranslateXY);
-		TranslationSubGizmos.Add(FSubGizmoInfo{ GizmoActor->TranslateXY, NewGizmo });
-	}
-
-	if (GizmoActor->RotateX != nullptr)
-	{
-		UInteractiveGizmo* NewGizmo = AddAxisRotationGizmo(GizmoActor->RotateX, GizmoComponent, AxisXSource, TransformSource, StateTarget);
-		ActiveComponents.Add(GizmoActor->RotateX);
-		RotationSubGizmos.Add(FSubGizmoInfo{ GizmoActor->RotateX, NewGizmo });
+		if (TranslateAxisComponents[AxisIndex])
+		{
+			Params.Component = TranslateAxisComponents[AxisIndex];
+			Params.Axis = Axes[AxisIndex];
+			AddAxisTranslationGizmo(Params, *SubGizmoSharedState);
+		}
 	}
-	if (GizmoActor->RotateY != nullptr)
+	UPrimitiveComponent* TranslatePlaneComponents[3]{ GizmoActor->TranslateYZ, GizmoActor->TranslateXZ, GizmoActor->TranslateXY };
+	for (int AxisIndex = 0; AxisIndex < 3; ++AxisIndex)
 	{
-		UInteractiveGizmo* NewGizmo = AddAxisRotationGizmo(GizmoActor->RotateY, GizmoComponent, AxisYSource, TransformSource, StateTarget);
-		ActiveComponents.Add(GizmoActor->RotateY);
-		RotationSubGizmos.Add(FSubGizmoInfo{ GizmoActor->RotateY, NewGizmo });
+		if (TranslatePlaneComponents[AxisIndex])
+		{
+			Params.Component = TranslatePlaneComponents[AxisIndex];
+			Params.Axis = Axes[AxisIndex];
+			AddPlaneTranslationGizmo(Params, *SubGizmoSharedState);
+		}
 	}
-	if (GizmoActor->RotateZ != nullptr)
+	UPrimitiveComponent* RotationAxisComponents[3]{ GizmoActor->RotateX, GizmoActor->RotateY, GizmoActor->RotateZ };
+	for (int AxisIndex = 0; AxisIndex < 3; ++AxisIndex)
 	{
-		UInteractiveGizmo* NewGizmo = AddAxisRotationGizmo(GizmoActor->RotateZ, GizmoComponent, AxisZSource, TransformSource, StateTarget);
-		ActiveComponents.Add(GizmoActor->RotateZ);
-		RotationSubGizmos.Add(FSubGizmoInfo{ GizmoActor->RotateZ, NewGizmo });
+		if (RotationAxisComponents[AxisIndex])
+		{
+			Params.Component = RotationAxisComponents[AxisIndex];
+			Params.Axis = Axes[AxisIndex];
+			AddAxisRotationGizmo(Params, *SubGizmoSharedState);
+		}
 	}
 	if (GizmoActor->RotationSphere != nullptr)
 	{
@@ -666,59 +893,43 @@ void UCombinedTransformGizmo::SetActiveTarget(UTransformProxy* Target, IToolCont
 		ActiveComponents.Add(GizmoActor->RotationSphere);
 		RotationSubGizmos.Add(FSubGizmoInfo{ GizmoActor->RotationSphere, nullptr });
 	}
-
-
-	// only need these if scaling enabled. Essentially these are just the unit axes, regardless
-	// of what 3D axis is in use, we will tell the ParameterSource-to-3D-Scale mapper to
-	// use the coordinate axes
-	UnitAxisXSource = UGizmoComponentAxisSource::Construct(GizmoComponent, 0, false, this);
-	UnitAxisYSource = UGizmoComponentAxisSource::Construct(GizmoComponent, 1, false, this);
-	UnitAxisZSource = UGizmoComponentAxisSource::Construct(GizmoComponent, 2, false, this);
-
 	if (GizmoActor->UniformScale != nullptr)
 	{
-		UInteractiveGizmo* NewGizmo = AddUniformScaleGizmo(GizmoActor->UniformScale, GizmoComponent, CameraAxisSource, CameraAxisSource, TransformSource, StateTarget);
-		ActiveComponents.Add(GizmoActor->UniformScale);
-		UniformScaleSubGizmos.Add(FSubGizmoInfo{ GizmoActor->UniformScale, NewGizmo });
+		Params.Component = GizmoActor->UniformScale;
+		Params.Axis = EAxis::None;
+		AddUniformScaleGizmo(Params, *SubGizmoSharedState);
 	}
-
-	if (GizmoActor->AxisScaleX != nullptr)
+	UPrimitiveComponent* ScaleAxisComponents[3]{ GizmoActor->AxisScaleX, GizmoActor->AxisScaleY, GizmoActor->AxisScaleZ };
+	for (int AxisIndex = 0; AxisIndex < 3; ++AxisIndex)
 	{
-		UInteractiveGizmo* NewGizmo = AddAxisScaleGizmo(GizmoActor->AxisScaleX, GizmoComponent, AxisXSource, UnitAxisXSource, TransformSource, StateTarget);
-		ActiveComponents.Add(GizmoActor->AxisScaleX);
-		NonUniformScaleSubGizmos.Add(FSubGizmoInfo{ GizmoActor->AxisScaleX, NewGizmo });
-	}
-	if (GizmoActor->AxisScaleY != nullptr)
-	{
-		UInteractiveGizmo* NewGizmo = AddAxisScaleGizmo(GizmoActor->AxisScaleY, GizmoComponent, AxisYSource, UnitAxisYSource, TransformSource, StateTarget);
-		ActiveComponents.Add(GizmoActor->AxisScaleY);
-		NonUniformScaleSubGizmos.Add(FSubGizmoInfo{ GizmoActor->AxisScaleY, NewGizmo });
+		if (ScaleAxisComponents[AxisIndex])
+		{
+			Params.Component = ScaleAxisComponents[AxisIndex];
+			Params.Axis = Axes[AxisIndex];
+			AddAxisScaleGizmo(Params, *SubGizmoSharedState);
+		}
 	}
-	if (GizmoActor->AxisScaleZ != nullptr)
+	UPrimitiveComponent* ScalePlaneComponents[3]{ GizmoActor->PlaneScaleYZ, GizmoActor->PlaneScaleXZ, GizmoActor->PlaneScaleXY };
+	for (int AxisIndex = 0; AxisIndex < 3; ++AxisIndex)
 	{
-		UInteractiveGizmo* NewGizmo = AddAxisScaleGizmo(GizmoActor->AxisScaleZ, GizmoComponent, AxisZSource, UnitAxisZSource, TransformSource, StateTarget);
-		ActiveComponents.Add(GizmoActor->AxisScaleZ);
-		NonUniformScaleSubGizmos.Add(FSubGizmoInfo{ GizmoActor->AxisScaleZ, NewGizmo });
+		if (ScalePlaneComponents[AxisIndex])
+		{
+			Params.Component = ScalePlaneComponents[AxisIndex];
+			Params.Axis = Axes[AxisIndex];
+			AddPlaneScaleGizmo(Params, *SubGizmoSharedState);
+		}
 	}
 
-	if (GizmoActor->PlaneScaleYZ != nullptr)
-	{
-		UInteractiveGizmo* NewGizmo = AddPlaneScaleGizmo(GizmoActor->PlaneScaleYZ, GizmoComponent, AxisXSource, UnitAxisXSource, TransformSource, StateTarget);
-		ActiveComponents.Add(GizmoActor->PlaneScaleYZ);
-		NonUniformScaleSubGizmos.Add(FSubGizmoInfo{ GizmoActor->PlaneScaleYZ, NewGizmo });
-	}
-	if (GizmoActor->PlaneScaleXZ != nullptr)
-	{
-		UInteractiveGizmo* NewGizmo = AddPlaneScaleGizmo(GizmoActor->PlaneScaleXZ, GizmoComponent, AxisYSource, UnitAxisYSource, TransformSource, StateTarget);
-		ActiveComponents.Add(GizmoActor->PlaneScaleXZ);
-		NonUniformScaleSubGizmos.Add(FSubGizmoInfo{ GizmoActor->PlaneScaleXZ, NewGizmo });
-	}
-	if (GizmoActor->PlaneScaleXY != nullptr)
-	{
-		UInteractiveGizmo* NewGizmo = AddPlaneScaleGizmo(GizmoActor->PlaneScaleXY, GizmoComponent, AxisZSource, UnitAxisZSource, TransformSource, StateTarget);
-		ActiveComponents.Add(GizmoActor->PlaneScaleXY);
-		NonUniformScaleSubGizmos.Add(FSubGizmoInfo{ GizmoActor->PlaneScaleXY, NewGizmo });
-	}
+	// Unpack the shared state into our properties. It might be nicer to just hold on to the shared state
+	//  object (in case it is needed later), but we do this for compatibility with existing child classes.
+	StateTarget = SubGizmoSharedState->StateTarget;
+	AxisXSource = SubGizmoSharedState->CardinalAxisSources[0];
+	AxisYSource = SubGizmoSharedState->CardinalAxisSources[1];
+	AxisZSource = SubGizmoSharedState->CardinalAxisSources[2];
+	CameraAxisSource = SubGizmoSharedState->CameraAxisSource;
+	UnitAxisXSource = SubGizmoSharedState->UnitCardinalAxisSources[0];
+	UnitAxisYSource = SubGizmoSharedState->UnitCardinalAxisSources[1];
+	UnitAxisZSource = SubGizmoSharedState->UnitCardinalAxisSources[2];
 
 	OnSetActiveTarget.Broadcast(this, ActiveTarget);
 }
@@ -835,6 +1046,147 @@ ETransformGizmoSubElements UCombinedTransformGizmo::GetGizmoElements()
 	return GetSubElementFlagsFromActor(GizmoActor);
 }
 
+UInteractiveGizmo* UCombinedTransformGizmo::AddAxisTranslationGizmo(
+	FTransformSubGizmoCommonParams& Params, FTransformSubGizmoSharedState& SharedState)
+{
+	UAxisPositionGizmo* Gizmo = Cast<UAxisPositionGizmo>(GetGizmoManager()->CreateGizmo(AxisPositionBuilderIdentifier));
+	if (!ensure(Gizmo))
+	{
+		return nullptr;
+	}
+
+	ensure(Gizmo->InitializeAsTranslateGizmo(Params, &SharedState));
+
+	if (UGizmoAxisTranslationParameterSource* ParamSource = Cast<UGizmoAxisTranslationParameterSource>(Gizmo->ParameterSource.GetObject()))
+	{
+		int AxisIndex = Params.GetClampedAxisIndex();
+		ParamSource->PositionConstraintFunction = [this](const FVector& Pos, FVector& Snapped) { return PositionSnapFunction(Pos, Snapped); };
+		ParamSource->AxisDeltaConstraintFunction = [this, AxisIndex](double AxisDelta, double& SnappedAxisDelta) { return PositionAxisDeltaSnapFunction(AxisDelta, SnappedAxisDelta, AxisIndex); };
+	}
+	else
+	{
+		ensure(false);
+	}
+
+	TranslationSubGizmos.Add(FSubGizmoInfo{ Params.Component, Gizmo });
+	ActiveComponents.Add(Params.Component);
+	ActiveGizmos.Add(Gizmo);
+
+	return Gizmo;
+}
+
+UInteractiveGizmo* UCombinedTransformGizmo::AddPlaneTranslationGizmo(
+	FTransformSubGizmoCommonParams& Params, FTransformSubGizmoSharedState& SharedState)
+{
+	UPlanePositionGizmo* Gizmo = Cast<UPlanePositionGizmo>(GetGizmoManager()->CreateGizmo(PlanePositionBuilderIdentifier));
+	if (!ensure(Gizmo))
+	{
+		return nullptr;
+	}
+	ensure(Gizmo->InitializeAsTranslateGizmo(Params, &SharedState));
+	
+	if (UGizmoPlaneTranslationParameterSource* ParamSource = Cast<UGizmoPlaneTranslationParameterSource>(Gizmo->ParameterSource.GetObject()))
+	{
+		int AxisIndex = Params.GetClampedAxisIndex();
+		int XAxes[3] = {1, 2, 0};
+		int YAxes[3] = {2, 0, 1};
+		ParamSource->PositionConstraintFunction = [this](const FVector& Pos, FVector& Snapped) { return PositionSnapFunction(Pos, Snapped); };
+		ParamSource->AxisXDeltaConstraintFunction = [this, XAxisIndex = XAxes[AxisIndex]](double AxisDelta, double& SnappedAxisDelta) { return PositionAxisDeltaSnapFunction(AxisDelta, SnappedAxisDelta, XAxisIndex); };
+		ParamSource->AxisYDeltaConstraintFunction = [this, YAxisIndex = YAxes[AxisIndex]](double AxisDelta, double& SnappedAxisDelta) { return PositionAxisDeltaSnapFunction(AxisDelta, SnappedAxisDelta, YAxisIndex); };
+	}
+	else
+	{
+		ensure(false);
+	}
+
+	TranslationSubGizmos.Add(FSubGizmoInfo{ Params.Component, Gizmo });
+	ActiveComponents.Add(Params.Component);
+	ActiveGizmos.Add(Gizmo);
+
+	return Gizmo;
+}
+
+UInteractiveGizmo* UCombinedTransformGizmo::AddAxisRotationGizmo(
+	FTransformSubGizmoCommonParams& Params,	FTransformSubGizmoSharedState& SharedState)
+{
+	UAxisAngleGizmo* Gizmo = Cast<UAxisAngleGizmo>(GetGizmoManager()->CreateGizmo(AxisAngleBuilderIdentifier));
+	if (!ensure(Gizmo))
+	{
+		return nullptr;
+	}
+	ensure(Gizmo->InitializeAsRotateGizmo(Params, &SharedState));
+
+	if (UGizmoAxisRotationParameterSource* AngleSource = Cast<UGizmoAxisRotationParameterSource>(Gizmo->AngleSource.GetObject()))
+	{
+		AngleSource->AngleDeltaConstraintFunction = [this](double AngleDelta, double& SnappedDelta) { return RotationAxisAngleSnapFunction(AngleDelta, SnappedDelta, 0); };
+	}
+	else
+	{
+		ensure(false);
+	}
+
+	RotationSubGizmos.Add(FSubGizmoInfo{ Params.Component, Gizmo });
+	ActiveComponents.Add(Params.Component);
+	ActiveGizmos.Add(Gizmo);
+
+	return Gizmo;
+}
+
+UInteractiveGizmo* UCombinedTransformGizmo::AddAxisScaleGizmo(
+	FTransformSubGizmoCommonParams& Params,	FTransformSubGizmoSharedState& SharedState)
+{
+	UAxisPositionGizmo* Gizmo = Cast<UAxisPositionGizmo>(GetGizmoManager()->CreateGizmo(AxisPositionBuilderIdentifier));
+	if (!ensure(Gizmo))
+	{
+		return nullptr;
+	}
+	ensure(Gizmo->InitializeAsScaleGizmo(Params, bDisallowNegativeScaling, &SharedState));
+
+	NonUniformScaleSubGizmos.Add(FSubGizmoInfo{ Params.Component, Gizmo });
+	ActiveComponents.Add(Params.Component);
+	ActiveGizmos.Add(Gizmo);
+
+	return Gizmo;
+}
+
+UInteractiveGizmo* UCombinedTransformGizmo::AddPlaneScaleGizmo(
+	FTransformSubGizmoCommonParams& Params,	FTransformSubGizmoSharedState& SharedState)
+{
+	UPlanePositionGizmo* Gizmo = Cast<UPlanePositionGizmo>(GetGizmoManager()->CreateGizmo(PlanePositionBuilderIdentifier));
+	if (!ensure(Gizmo))
+	{
+		return nullptr;
+	}
+	ensure(Gizmo->InitializeAsScaleGizmo(Params, bDisallowNegativeScaling, &SharedState));
+
+	NonUniformScaleSubGizmos.Add(FSubGizmoInfo{ Params.Component, Gizmo });
+	ActiveComponents.Add(Params.Component);
+	ActiveGizmos.Add(Gizmo);
+
+	return Gizmo;
+}
+
+UInteractiveGizmo* UCombinedTransformGizmo::AddUniformScaleGizmo(
+	FTransformSubGizmoCommonParams& Params,	FTransformSubGizmoSharedState& SharedState)
+{
+	UPlanePositionGizmo* Gizmo = Cast<UPlanePositionGizmo>(GetGizmoManager()->CreateGizmo(PlanePositionBuilderIdentifier));
+	if (!ensure(Gizmo))
+	{
+		return nullptr;
+	}
+	ensure(Gizmo->InitializeAsUniformScaleGizmo(Params, bDisallowNegativeScaling, &SharedState));
+
+	UniformScaleSubGizmos.Add(FSubGizmoInfo{ Params.Component, Gizmo });
+	ActiveComponents.Add(Params.Component);
+	ActiveGizmos.Add(Gizmo);
+
+	return Gizmo;
+}
+
+
+// These are deprecated initialization functions that do sub gizmo initialization by hand instead of using the
+//  "InitializeAs..." functions that were added to subgizmos to make them simpler to instantiate outside of
+//  this class.
 UInteractiveGizmo* UCombinedTransformGizmo::AddAxisTranslationGizmo(
 	UPrimitiveComponent* AxisComponent, USceneComponent* RootComponent,
 	IGizmoAxisSource* AxisSource,
@@ -874,9 +1226,6 @@ UInteractiveGizmo* UCombinedTransformGizmo::AddAxisTranslationGizmo(
 	ActiveGizmos.Add(TranslateGizmo);
 	return TranslateGizmo;
 }
-
-
-
 UInteractiveGizmo* UCombinedTransformGizmo::AddPlaneTranslationGizmo(
 	UPrimitiveComponent* AxisComponent, USceneComponent* RootComponent,
 	IGizmoAxisSource* AxisSource,
@@ -917,11 +1266,6 @@ UInteractiveGizmo* UCombinedTransformGizmo::AddPlaneTranslationGizmo(
 	ActiveGizmos.Add(TranslateGizmo);
 	return TranslateGizmo;
 }
-
-
-
-
-
 UInteractiveGizmo* UCombinedTransformGizmo::AddAxisRotationGizmo(
 	UPrimitiveComponent* AxisComponent, USceneComponent* RootComponent,
 	IGizmoAxisSource* AxisSource,
@@ -962,9 +1306,6 @@ UInteractiveGizmo* UCombinedTransformGizmo::AddAxisRotationGizmo(
 
 	return RotateGizmo;
 }
-
-
-
 UInteractiveGizmo* UCombinedTransformGizmo::AddAxisScaleGizmo(
 	UPrimitiveComponent* AxisComponent, USceneComponent* RootComponent,
 	IGizmoAxisSource* GizmoAxisSource, IGizmoAxisSource* ParameterAxisSource,
@@ -997,9 +1338,6 @@ UInteractiveGizmo* UCombinedTransformGizmo::AddAxisScaleGizmo(
 	ActiveGizmos.Add(ScaleGizmo);
 	return ScaleGizmo;
 }
-
-
-
 UInteractiveGizmo* UCombinedTransformGizmo::AddPlaneScaleGizmo(
 	UPrimitiveComponent* AxisComponent, USceneComponent* RootComponent,
 	IGizmoAxisSource* GizmoAxisSource, IGizmoAxisSource* ParameterAxisSource,
@@ -1033,11 +1371,6 @@ UInteractiveGizmo* UCombinedTransformGizmo::AddPlaneScaleGizmo(
 	ActiveGizmos.Add(ScaleGizmo);
 	return ScaleGizmo;
 }
-
-
-
-
-
 UInteractiveGizmo* UCombinedTransformGizmo::AddUniformScaleGizmo(
 	UPrimitiveComponent* ScaleComponent, USceneComponent* RootComponent,
 	IGizmoAxisSource* GizmoAxisSource, IGizmoAxisSource* ParameterAxisSource,
@@ -1097,6 +1430,7 @@ void UCombinedTransformGizmo::ClearActiveTarget()
 	StateTarget = nullptr;
 
 	ActiveTarget = nullptr;
+	TransactionProviderAtLastSetActiveTarget = nullptr;
 }
 
 
diff --git a/Engine/Source/Runtime/InteractiveToolsFramework/Private/BaseGizmos/GizmoPrivateUtil.cpp b/Engine/Source/Runtime/InteractiveToolsFramework/Private/BaseGizmos/GizmoPrivateUtil.cpp
--- a/Engine/Source/Runtime/InteractiveToolsFramework/Private/BaseGizmos/GizmoPrivateUtil.cpp
+++ b/Engine/Source/Runtime/InteractiveToolsFramework/Private/BaseGizmos/GizmoPrivateUtil.cpp
@@ -2,9 +2,51 @@
 
 #include "BaseGizmos/GizmoPrivateUtil.h"
 
+#include "BaseGizmos/AxisSources.h"
+#include "BaseGizmos/AxisAngleGizmo.h"
+#include "BaseGizmos/AxisPositionGizmo.h"
+#include "BaseGizmos/CombinedTransformGizmo.h"
+#include "BaseGizmos/GizmoBaseComponent.h"
 #include "BaseGizmos/GizmoViewContext.h"
+#include "BaseGizmos/HitTargets.h"
+#include "BaseGizmos/PlanePositionGizmo.h"
+#include "BaseGizmos/StateTargets.h"
+#include "BaseGizmos/TransformSources.h"
+#include "BaseGizmos/TransformSubGizmoUtil.h" // FTransformSubGizmoCommonParams
 #include "ContextObjectStore.h"
+#include "GameFramework/Actor.h"
 #include "InteractiveGizmoManager.h"
+#include "Internationalization/Internationalization.h" // LOCTEXT
+#include "UObject/UObjectGlobals.h" // GetTransientPackage
+
+EAxis::Type UE::GizmoUtil::ToAxis(ETransformGizmoSubElements Element)
+{
+	switch (Element)
+	{
+	case ETransformGizmoSubElements::TranslateAxisX:
+	case ETransformGizmoSubElements::RotateAxisX:
+	case ETransformGizmoSubElements::ScaleAxisX:
+	case ETransformGizmoSubElements::TranslatePlaneYZ:
+	case ETransformGizmoSubElements::ScalePlaneYZ:
+		return EAxis::X;
+	case ETransformGizmoSubElements::TranslateAxisY:
+	case ETransformGizmoSubElements::RotateAxisY:
+	case ETransformGizmoSubElements::ScaleAxisY:
+	case ETransformGizmoSubElements::TranslatePlaneXZ:
+	case ETransformGizmoSubElements::ScalePlaneXZ:
+		return EAxis::Y;
+	case ETransformGizmoSubElements::TranslateAxisZ:
+	case ETransformGizmoSubElements::RotateAxisZ:
+	case ETransformGizmoSubElements::ScaleAxisZ:
+	case ETransformGizmoSubElements::TranslatePlaneXY:
+	case ETransformGizmoSubElements::ScalePlaneXY:
+		return EAxis::Z;
+	default:
+		// We don't ensure here because it is sometimes convenient to write code that ends up
+		//  passing in things that we don't end up using the axis for, like a uniform scale element.
+		return EAxis::None;
+	}
+}
 
 UGizmoViewContext* UE::GizmoUtil::GetGizmoViewContext(UInteractiveGizmoManager* GizmoManager)
 {
@@ -24,3 +66,150 @@ UGizmoViewContext* UE::GizmoUtil::GetGizmoViewContext(UInteractiveGizmoManager*
 
 	return GizmoViewContext;
 }
+
+template <typename SubGizmoType>
+bool UE::GizmoUtil::SetCommonSubGizmoProperties(
+	SubGizmoType* Gizmo,
+	const UE::GizmoUtil::FTransformSubGizmoCommonParams& Params,
+	UE::GizmoUtil::FTransformSubGizmoSharedState* SharedState,
+	UGizmoScaledAndUnscaledTransformSources*& TransformSourceOut)
+{
+	if (!Gizmo
+		|| !Params.Component
+		|| !Params.TransformProxy)
+	{
+		return false;
+	}
+
+	AActor* ComponentOwnerActor = Params.Component->GetOwner();
+	UObject* OuterForSubobjects = Params.OuterForSubobjects ? Params.OuterForSubobjects : GetTransientPackage();
+
+	// Set up the axis source if we got an axis
+	if (Params.Axis != EAxis::None)
+	{
+		int AxisIndex = Params.GetClampedAxisIndex();
+		if (!Params.bAxisIsBasedOnRootComponent)
+		{
+			// Axis will be based on the passed-in component transform
+			Gizmo->AxisSource = UGizmoComponentAxisSource::Construct(Params.Component, AxisIndex,
+				// bUseLocalAxes, not important because we're going to be updating this value if needed
+				true,
+				OuterForSubobjects);
+		}
+		else if (SharedState && SharedState->CardinalAxisSources[AxisIndex])
+		{
+			// Axis is based off root component, and we already have one in our shared state
+			Gizmo->AxisSource = SharedState->CardinalAxisSources[AxisIndex];
+		}
+		else if (ensure(ComponentOwnerActor && ComponentOwnerActor->GetRootComponent()))
+		{
+			// Axis is based off root component, but we haven't created one yet
+			UGizmoComponentAxisSource* CastAxisSource = UGizmoComponentAxisSource::Construct(ComponentOwnerActor->GetRootComponent(), AxisIndex,
+				// bUseLocalAxes, not important because we're going to be updating this value if needed
+				true,
+				OuterForSubobjects);
+			Gizmo->AxisSource = CastAxisSource;
+			if (SharedState)
+			{
+				SharedState->CardinalAxisSources[AxisIndex] = CastAxisSource;
+			}
+		}
+	}
+
+	USceneComponent* ComponentToMove = Params.Component;
+	if (Params.bManipulatesRootComponent && ensure(ComponentOwnerActor && ComponentOwnerActor->GetRootComponent()))
+	{
+		ComponentToMove = ComponentOwnerActor->GetRootComponent();
+	}
+
+	// The transform source is also transform destination. It forwards the resulting transform to our proxy and
+	// to the component.
+	if (SharedState && SharedState->TransformSource && Params.bManipulatesRootComponent)
+	{
+		// We already have one in our shared state and we can use it (because the sub gizmo doesn't just move
+		//  its own component).
+		TransformSourceOut = SharedState->TransformSource;
+	}
+	else
+	{
+		// Need to create a new transform source
+		TransformSourceOut = UGizmoScaledAndUnscaledTransformSources::Construct(
+			UGizmoTransformProxyTransformSource::Construct(Params.TransformProxy, OuterForSubobjects),
+			ComponentToMove, OuterForSubobjects);
+		if (SharedState && Params.bManipulatesRootComponent)
+		{
+			SharedState->TransformSource = TransformSourceOut;
+		}
+	}
+
+	// Hit target is how we detect whether the gizmo has been hit, and what we update hover on. We don't
+	//  use shared state here because this is always unique to each sub gizmo.
+	UGizmoComponentHitTarget* CastHitTarget = UGizmoComponentHitTarget::Construct(Params.Component, OuterForSubobjects);
+	Gizmo->HitTarget = CastHitTarget;
+	// The default hover/interaction behavior is to forward that information to the component for rendering.
+	CastHitTarget->UpdateHoverFunction = 
+		// TODO: See comment in FTransformSubGizmoCommonParams- seems not worth giving the user the ability to
+		//  customize this function.
+		//Params.CustomUpdateHoverFunction ? Params.CustomUpdateHoverFunction : 
+		[Component = Params.Component](bool bHovering)
+	{
+		if (IGizmoBaseComponentInterface* Hoverable = Cast<IGizmoBaseComponentInterface>(Component))
+		{
+			Hoverable->UpdateHoverState(bHovering);
+		}
+	};
+	CastHitTarget->UpdateInteractingFunction = [Component = Params.Component](bool bInteracting)
+	{
+		if (IGizmoBaseComponentInterface* Hoverable = Cast<IGizmoBaseComponentInterface>(Component))
+		{
+			Hoverable->UpdateInteractingState(bInteracting);
+		}
+	};
+
+	// Set up shared state target, which handles undo/redo of the component and proxy transform. As
+	//  with transform source, if the sub gizmo moves its own component (rather than the whole gizmo),
+	//  we need our own target, otherwise we'll try to use shared state.
+	if (SharedState && SharedState->StateTarget && Params.bManipulatesRootComponent)
+	{
+		Gizmo->StateTarget = SharedState->StateTarget;
+	}
+	else
+	{
+		const FText DefaultTransactionName(NSLOCTEXT("UCombinedTransformGizmo", "UCombinedTransformGizmoTransaction", "Transform"));
+
+		FText TransactionName = Params.TransactionName.IsSet() ? Params.TransactionName.GetValue()
+			: DefaultTransactionName;
+
+		IToolContextTransactionProvider* TransactionProvider = Params.TransactionProvider ?
+			Params.TransactionProvider : Gizmo->GetGizmoManager();
+
+		UGizmoTransformChangeStateTarget* CastStateTarget = UGizmoTransformChangeStateTarget::Construct(ComponentToMove,
+			TransactionName, TransactionProvider, OuterForSubobjects);
+		CastStateTarget->DependentChangeSources.Add(MakeUnique<FTransformProxyChangeSource>(Params.TransformProxy));
+		Gizmo->StateTarget = CastStateTarget;
+
+		if (SharedState && Params.bManipulatesRootComponent)
+		{
+			SharedState->StateTarget = CastStateTarget;
+		}
+	}
+
+	return true;
+}
+
+// Explicit instantiations to avoid having the definition in the header
+template bool UE::GizmoUtil::SetCommonSubGizmoProperties<UAxisPositionGizmo>(
+	UAxisPositionGizmo* Gizmo,
+	const FTransformSubGizmoCommonParams& Params,
+	FTransformSubGizmoSharedState* SharedState,
+	UGizmoScaledAndUnscaledTransformSources*& TransformSourceOut);
+template bool UE::GizmoUtil::SetCommonSubGizmoProperties<UAxisAngleGizmo>(
+	UAxisAngleGizmo* Gizmo,
+	const FTransformSubGizmoCommonParams& Params,
+	FTransformSubGizmoSharedState* SharedState,
+	UGizmoScaledAndUnscaledTransformSources*& TransformSourceOut);
+template bool UE::GizmoUtil::SetCommonSubGizmoProperties<UPlanePositionGizmo>(
+	UPlanePositionGizmo* Gizmo,
+	const FTransformSubGizmoCommonParams& Params,
+	FTransformSubGizmoSharedState* SharedState,
+	UGizmoScaledAndUnscaledTransformSources*& TransformSourceOut);
\ No newline at end of file
diff --git a/Engine/Source/Runtime/InteractiveToolsFramework/Private/BaseGizmos/GizmoPrivateUtil.h b/Engine/Source/Runtime/InteractiveToolsFramework/Private/BaseGizmos/GizmoPrivateUtil.h
--- a/Engine/Source/Runtime/InteractiveToolsFramework/Private/BaseGizmos/GizmoPrivateUtil.h
+++ b/Engine/Source/Runtime/InteractiveToolsFramework/Private/BaseGizmos/GizmoPrivateUtil.h
@@ -2,6 +2,10 @@
 
 #pragma once
 
+#include "InteractiveGizmo.h" // ETransformGizmoSubElements
+#include "Math/Axis.h"
+
+class UGizmoScaledAndUnscaledTransformSources;
 class UGizmoViewContext;
 class UInteractiveGizmoManager;
 
@@ -18,6 +22,11 @@ namespace UE::GizmoUtil
 
 namespace UE::GizmoUtil 
 {
+	/**
+	 * Given a single element, gives the axis that defines that element (e.g., x for TranslateAxisX or TranslatePlaneYZ, etc).
+	 *  Gives EAxis::None if the element is not a single element.
+	 */
+	EAxis::Type ToAxis(ETransformGizmoSubElements Element);
 
 	/**
 	 * Simple helper that gets the gizmo view context out of the context object store associated with a gizmo manager.
@@ -25,4 +34,19 @@ namespace UE::GizmoUtil
 	 */
 	UGizmoViewContext* GetGizmoViewContext(UInteractiveGizmoManager* GizmoManager);
 
+	/**
+	 * Helper that sets some common properties on sub gizmos, namely AxisSource, HitTarget, 
+	 *  and StateTarget. Also returns a TransformSource, which typically gets wrapped up in
+	 *  a parameter source.
+	 * 
+	 * Template because the properties it manipulates aren't part of a base class. Perhaps they
+	 *  should be, but we have not needed that yet, other than this.
+	 */
+	template <typename SubGizmoType>
+	bool SetCommonSubGizmoProperties(
+		SubGizmoType* Gizmo,
+		const UE::GizmoUtil::FTransformSubGizmoCommonParams& Params,
+		UE::GizmoUtil::FTransformSubGizmoSharedState* SharedState,
+		UGizmoScaledAndUnscaledTransformSources*& TransformSourceOut);
+
 }// end UE::GizmoUtil
\ No newline at end of file
diff --git a/Engine/Source/Runtime/InteractiveToolsFramework/Private/BaseGizmos/PlanePositionGizmo.cpp b/Engine/Source/Runtime/InteractiveToolsFramework/Private/BaseGizmos/PlanePositionGizmo.cpp
--- a/Engine/Source/Runtime/InteractiveToolsFramework/Private/BaseGizmos/PlanePositionGizmo.cpp
+++ b/Engine/Source/Runtime/InteractiveToolsFramework/Private/BaseGizmos/PlanePositionGizmo.cpp
@@ -5,9 +5,31 @@
 #include "BaseBehaviors/ClickDragBehavior.h"
 #include "BaseBehaviors/MouseHoverBehavior.h"
 #include "BaseGizmos/GizmoMath.h"
+#include "BaseGizmos/GizmoPrivateUtil.h" // SetCommonSubGizmoProperties
+#include "BaseGizmos/TransformSubGizmoUtil.h" // FTransformSubGizmoCommonParams
+#include "Components/PrimitiveComponent.h"
 
 #include UE_INLINE_GENERATED_CPP_BY_NAME(PlanePositionGizmo)
 
+namespace PlanePositionGizmoLocals
+{
+	bool UpdateCameraAxisSource(UGizmoConstantFrameAxisSource& CameraAxisSource, UInteractiveGizmoManager* GizmoManager, const FVector3d& AxisOrigin)
+	{
+		if (!GizmoManager || !GizmoManager->GetContextQueriesAPI())
+		{
+			return false;
+		}
+
+		FViewCameraState CameraState;
+		GizmoManager->GetContextQueriesAPI()->GetCurrentViewState(CameraState);
+
+		CameraAxisSource.Origin = AxisOrigin;
+		CameraAxisSource.Direction = -CameraState.Forward();
+		CameraAxisSource.TangentX = CameraState.Right();
+		CameraAxisSource.TangentY = CameraState.Up();
+		return true;
+	}
+}
 
 
 UInteractiveGizmo* UPlanePositionGizmoBuilder::BuildGizmo(const FToolBuilderState& SceneState) const
@@ -41,6 +63,160 @@ void UPlanePositionGizmo::Setup()
 	bInInteraction = false;
 }
 
+bool UPlanePositionGizmo::InitializeAsTranslateGizmo(
+	const UE::GizmoUtil::FTransformSubGizmoCommonParams& Params,
+	UE::GizmoUtil::FTransformSubGizmoSharedState* SharedState)
+{
+	if (!Params.Component
+		|| !Params.TransformProxy
+		|| Params.Axis == EAxis::None)
+	{
+		return false;
+	}
+
+	UGizmoScaledAndUnscaledTransformSources* TransformSource;
+	if (!UE::GizmoUtil::SetCommonSubGizmoProperties(this, Params, SharedState, TransformSource))
+	{
+		return false;
+	}
+
+	UObject* Owner = Params.OuterForSubobjects ? Params.OuterForSubobjects : GetTransientPackage();
+
+	// Parameter source maps axis-parameter-change to translation of TransformSource's transform
+	ParameterSource = UGizmoPlaneTranslationParameterSource::Construct(
+		AxisSource.GetInterface(), TransformSource, Owner);
+
+	return true;
+}
+
+bool UPlanePositionGizmo::InitializeAsScaleGizmo(
+	const UE::GizmoUtil::FTransformSubGizmoCommonParams& Params, bool bDisallowNegativeScaling,
+	UE::GizmoUtil::FTransformSubGizmoSharedState* SharedState)
+{
+	if (!Params.Component
+		|| !Params.TransformProxy
+		|| Params.Axis == EAxis::None)
+	{
+		return false;
+	}
+
+	int AxisIndex = Params.GetClampedAxisIndex();
+
+	UGizmoScaledAndUnscaledTransformSources* TransformSource;
+	if (!UE::GizmoUtil::SetCommonSubGizmoProperties(this, Params, SharedState, TransformSource))
+	{
+		return false;
+	}
+
+	UObject* Owner = Params.OuterForSubobjects ? Params.OuterForSubobjects : GetTransientPackage();
+	bEnableSignedAxis = true;
+
+	// Although the normal axis source gets used for detecting interactions, the parameter application has
+	//  to happen along unrotated axes because the scaling gets applied before rotation. In other words if we
+	//  tried to apply scaling measured along a rotated vector, we would end up incorrectly scaling along
+	//  multiple axes.
+	UGizmoComponentAxisSource* UnitCardinalAxisSource = nullptr;
+	// See if we already have it in our shared state
+	if (SharedState && SharedState->UnitCardinalAxisSources[AxisIndex])
+	{
+		UnitCardinalAxisSource = SharedState->UnitCardinalAxisSources[AxisIndex];
+	}
+	else
+	{
+		// Create new and add to shared state.
+		USceneComponent* RootComponent = Params.Component->GetOwner()->GetRootComponent();
+		UGizmoComponentAxisSource* CastAxisSource = UGizmoComponentAxisSource::Construct(RootComponent, AxisIndex,
+			// bUseLocalAxes, not important because we're going to be updating this value every tick
+			true,
+			Owner);
+		UnitCardinalAxisSource = CastAxisSource;
+		if (SharedState)
+		{
+			SharedState->UnitCardinalAxisSources[AxisIndex] = UnitCardinalAxisSource;
+		}
+	}
+
+	// Parameter source maps axis-parameter-change to scale of TransformSource's transform
+	UGizmoPlaneScaleParameterSource* CastParameterSource = UGizmoPlaneScaleParameterSource::Construct(
+		UnitCardinalAxisSource, TransformSource, Owner);
+	ParameterSource = CastParameterSource;
+	CastParameterSource->bClampToZero = bDisallowNegativeScaling;
+	CastParameterSource->bUseEqualScaling = true;
+
+	return true;
+}
+
+bool UPlanePositionGizmo::InitializeAsUniformScaleGizmo(
+	const UE::GizmoUtil::FTransformSubGizmoCommonParams& Params,
+	bool bDisallowNegativeScaling,
+	UE::GizmoUtil::FTransformSubGizmoSharedState* SharedState)
+{
+	using namespace PlanePositionGizmoLocals;
+
+	if (!Params.Component
+		|| !Params.TransformProxy)
+	{
+		return false;
+	}
+
+	UGizmoScaledAndUnscaledTransformSources* TransformSource;
+
+	// Make sure axis index is invalid so that the SetCommonSubGizmoProperties call below doesn't create
+	// an axis source for us.
+	if (!ensureMsgf(Params.Axis == EAxis::None, TEXT("InitializeAsUniformScaleGizmo uses a camera axis source.")))
+	{
+		UE::GizmoUtil::FTransformSubGizmoCommonParams ParamsCopy = Params;
+		ParamsCopy.Axis = EAxis::None;
+		if (!UE::GizmoUtil::SetCommonSubGizmoProperties(this, ParamsCopy, SharedState, TransformSource))
+		{
+			return false;
+		}
+	}
+	else if (!UE::GizmoUtil::SetCommonSubGizmoProperties(this, Params, SharedState, TransformSource))
+	{
+		return false;
+	}
+
+	UObject* Owner = Params.OuterForSubobjects ? Params.OuterForSubobjects : GetTransientPackage();
+
+	UGizmoConstantFrameAxisSource* CameraAxisSource = nullptr;
+
+	// See if we already have it in our shared state
+	if (SharedState && SharedState->CameraAxisSource)
+	{
+		CameraAxisSource = SharedState->CameraAxisSource;
+	}
+	else
+	{
+		// Create new and add to shared state.
+		CameraAxisSource = NewObject<UGizmoConstantFrameAxisSource>(Owner);
+		TWeakObjectPtr<AActor> OwnerActor = Params.Component->GetOwner();
+		CustomTickFunction = [this, OwnerActor](float DeltaTime)
+		{
+			if (UGizmoConstantFrameAxisSource* FrameAxisSource = Cast<UGizmoConstantFrameAxisSource>(AxisSource.GetObject()))
+			{
+				UpdateCameraAxisSource(
+					*FrameAxisSource,
+					GetGizmoManager(),
+					OwnerActor.IsValid() ? OwnerActor->GetTransform().GetLocation() : FVector3d::ZeroVector);
+			}
+		};
+
+		if (SharedState)
+		{
+			SharedState->CameraAxisSource = CameraAxisSource;
+		}
+	}
+	AxisSource = CameraAxisSource;
+
+	// Parameter source maps axis-parameter-change to scale of TransformSource's transform
+	UGizmoUniformScaleParameterSource* CastParameterSource = UGizmoUniformScaleParameterSource::Construct(
+		CameraAxisSource, TransformSource, Owner);
+	ParameterSource = CastParameterSource;
+
+	return true;
+}
+
 FInputRayHit UPlanePositionGizmo::CanBeginClickDragSequence(const FInputDeviceRay& PressPos)
 {
 	FInputRayHit GizmoHit;
@@ -113,6 +289,10 @@ void UPlanePositionGizmo::OnClickPress(const FInputDeviceRay& PressPos)
 	{
 		StateTarget->BeginUpdate();
 	}
+	if (ensure(HitTarget))
+	{
+		HitTarget->UpdateInteractingState(bInInteraction);
+	}
 }
 
 void UPlanePositionGizmo::OnClickDrag(const FInputDeviceRay& DragPos)
@@ -175,6 +355,10 @@ void UPlanePositionGizmo::OnClickRelease(const FInputDeviceRay& ReleasePos)
 		StateTarget->EndUpdate();
 	}
 	bInInteraction = false;
+	if (ensure(HitTarget))
+	{
+		HitTarget->UpdateInteractingState(bInInteraction);
+	}
 }
 
 
@@ -191,6 +375,10 @@ void UPlanePositionGizmo::OnTerminateDragSequence()
 		StateTarget->EndUpdate();
 	}
 	bInInteraction = false;
+	if (ensure(HitTarget))
+	{
+		HitTarget->UpdateInteractingState(bInInteraction);
+	}
 }
 
 
@@ -224,3 +412,11 @@ void UPlanePositionGizmo::OnEndHover()
 	HitTarget->UpdateHoverState(false);
 }
 
+
+void UPlanePositionGizmo::Tick(float DeltaTime)
+{
+	if (CustomTickFunction)
+	{
+		CustomTickFunction(DeltaTime);
+	}
+}
diff --git a/Engine/Source/Runtime/InteractiveToolsFramework/Private/BaseGizmos/RepositionableTransformGizmo.cpp b/Engine/Source/Runtime/InteractiveToolsFramework/Private/BaseGizmos/RepositionableTransformGizmo.cpp
--- a/Engine/Source/Runtime/InteractiveToolsFramework/Private/BaseGizmos/RepositionableTransformGizmo.cpp
+++ b/Engine/Source/Runtime/InteractiveToolsFramework/Private/BaseGizmos/RepositionableTransformGizmo.cpp
@@ -10,6 +10,7 @@
 #include "BaseGizmos/PlanePositionGizmo.h"
 #include "BaseGizmos/TransformProxy.h"
 #include "BaseGizmos/TransformSources.h"
+#include "BaseGizmos/TransformSubGizmoUtil.h"
 #include "ContextObjectStore.h"
 
 #include "InteractiveGizmo.h"
@@ -71,67 +72,65 @@ void URepositionableTransformGizmo::SetActiveTarget(UTransformProxy* Target, ITo
 	ProxyChangeSource->bOverrideSetPivotMode = true;
 	RepositionStateTarget->DependentChangeSources.Add(MoveTemp(ProxyChangeSource));
 
+	// Pack the above up for use in the initialization functions
+	FTransformSubGizmoSharedState SharedStateToUse;
+	SharedStateToUse.TransformSource = TransformSource;
+	SharedStateToUse.StateTarget = RepositionStateTarget;
+	// The shared data struct should have been created in the base class SetActiveTarget
+	if (ensure(SubGizmoSharedState.IsValid()))
+	{
+		for (int i = 0; i < 3; ++i)
+		{
+			SharedStateToUse.CardinalAxisSources[i] = SubGizmoSharedState->CardinalAxisSources[i];
+		}
+	}
+
 	// Add on the extra gizmos to existing components. We use the base class addition functions and
 	// then go back to tweak the parts that need changing using our own "Modify..." functions.
 	// We also store the repositioning gizmos in an additional separate list so that we can
 	// modify their alignment functions separately if we need to.
-	if (GizmoActor->TranslateX != nullptr)
-	{
-		UInteractiveGizmo* SubGizmo = AddAxisTranslationGizmo(GizmoActor->TranslateX, GizmoComponent, AxisXSource, TransformSource, RepositionStateTarget, 0);
-		ModifyPivotAxisGizmo(SubGizmo);
-		PivotAlignmentGizmos.Add(SubGizmo);
-	}
-	if (GizmoActor->TranslateY != nullptr)
-	{
-		UInteractiveGizmo* SubGizmo = AddAxisTranslationGizmo(GizmoActor->TranslateY, GizmoComponent, AxisYSource, TransformSource, RepositionStateTarget, 1);
-		ModifyPivotAxisGizmo(SubGizmo);
-		PivotAlignmentGizmos.Add(SubGizmo);
-	}
-	if (GizmoActor->TranslateZ != nullptr)
-	{
-		UInteractiveGizmo* SubGizmo = AddAxisTranslationGizmo(GizmoActor->TranslateZ, GizmoComponent, AxisZSource, TransformSource, RepositionStateTarget, 2);
-		ModifyPivotAxisGizmo(SubGizmo);
-		PivotAlignmentGizmos.Add(SubGizmo);
-	}
 
+	UE::GizmoUtil::FTransformSubGizmoCommonParams Params;
+	Params.TransformProxy = ActiveTarget;
+	Params.TransactionProvider = TransactionProvider;
+	Params.bManipulatesRootComponent = true;
 
-	if (GizmoActor->TranslateYZ != nullptr)
-	{
-		UInteractiveGizmo* SubGizmo = AddPlaneTranslationGizmo(GizmoActor->TranslateYZ, GizmoComponent, AxisXSource, TransformSource, RepositionStateTarget, 1, 2);
-		ModifyPivotPlaneGizmo(SubGizmo);
-		PivotAlignmentGizmos.Add(SubGizmo);
-	}
-	if (GizmoActor->TranslateXZ != nullptr)
+	EAxis::Type Axes[3] = { EAxis::X, EAxis::Y, EAxis::Z };
+	UPrimitiveComponent* TranslateAxisComponents[3]{ GizmoActor->TranslateX, GizmoActor->TranslateY, GizmoActor->TranslateZ };
+	for (int AxisIndex = 0; AxisIndex < 3; ++AxisIndex)
 	{
-		UInteractiveGizmo* SubGizmo = AddPlaneTranslationGizmo(GizmoActor->TranslateXZ, GizmoComponent, AxisYSource, TransformSource, RepositionStateTarget, 2, 0);
-		ModifyPivotPlaneGizmo(SubGizmo);
-		PivotAlignmentGizmos.Add(SubGizmo);
-	}
-	if (GizmoActor->TranslateXY != nullptr)
-	{
-		UInteractiveGizmo* SubGizmo = AddPlaneTranslationGizmo(GizmoActor->TranslateXY, GizmoComponent, AxisZSource, TransformSource, RepositionStateTarget, 0, 1);
-		ModifyPivotPlaneGizmo(SubGizmo);
-		PivotAlignmentGizmos.Add(SubGizmo);
-	}
-
-
-	if (GizmoActor->RotateX != nullptr)
-	{
-		UInteractiveGizmo* SubGizmo = AddAxisRotationGizmo(GizmoActor->RotateX, GizmoComponent, AxisXSource, TransformSource, RepositionStateTarget);
-		ModifyPivotRotateGizmo(SubGizmo);
-		PivotAlignmentGizmos.Add(SubGizmo);
+		if (TranslateAxisComponents[AxisIndex])
+		{
+			Params.Component = TranslateAxisComponents[AxisIndex];
+			Params.Axis = Axes[AxisIndex];
+			UInteractiveGizmo* SubGizmo = AddAxisTranslationGizmo(Params, SharedStateToUse);
+			ModifyPivotAxisGizmo(SubGizmo);
+			PivotAlignmentGizmos.Add(SubGizmo);
+		}
 	}
-	if (GizmoActor->RotateY != nullptr)
+	UPrimitiveComponent* TranslatePlaneComponents[3]{ GizmoActor->TranslateYZ, GizmoActor->TranslateXZ, GizmoActor->TranslateXY };
+	for (int AxisIndex = 0; AxisIndex < 3; ++AxisIndex)
 	{
-		UInteractiveGizmo* SubGizmo = AddAxisRotationGizmo(GizmoActor->RotateY, GizmoComponent, AxisYSource, TransformSource, RepositionStateTarget);
-		ModifyPivotRotateGizmo(SubGizmo);
-		PivotAlignmentGizmos.Add(SubGizmo);
+		if (TranslatePlaneComponents[AxisIndex])
+		{
+			Params.Component = TranslatePlaneComponents[AxisIndex];
+			Params.Axis = Axes[AxisIndex];
+			UInteractiveGizmo* SubGizmo = AddPlaneTranslationGizmo(Params, SharedStateToUse);
+			ModifyPivotPlaneGizmo(SubGizmo);
+			PivotAlignmentGizmos.Add(SubGizmo);
+		}
 	}
-	if (GizmoActor->RotateZ != nullptr)
+	UPrimitiveComponent* RotationAxisComponents[3]{ GizmoActor->RotateX, GizmoActor->RotateY, GizmoActor->RotateZ };
+	for (int AxisIndex = 0; AxisIndex < 3; ++AxisIndex)
 	{
-		UInteractiveGizmo* SubGizmo = AddAxisRotationGizmo(GizmoActor->RotateZ, GizmoComponent, AxisZSource, TransformSource, RepositionStateTarget);
-		ModifyPivotRotateGizmo(SubGizmo);
-		PivotAlignmentGizmos.Add(SubGizmo);
+		if (RotationAxisComponents[AxisIndex])
+		{
+			Params.Component = RotationAxisComponents[AxisIndex];
+			Params.Axis = Axes[AxisIndex];
+			UInteractiveGizmo* SubGizmo = AddAxisRotationGizmo(Params, SharedStateToUse);
+			ModifyPivotRotateGizmo(SubGizmo);
+			PivotAlignmentGizmos.Add(SubGizmo);
+		}
 	}
 }
 
diff --git a/Engine/Source/Runtime/InteractiveToolsFramework/Private/BaseGizmos/TransformSubGizmoUtil.cpp b/Engine/Source/Runtime/InteractiveToolsFramework/Private/BaseGizmos/TransformSubGizmoUtil.cpp
new file mode 100644
--- /dev/null
+++ b/Engine/Source/Runtime/InteractiveToolsFramework/Private/BaseGizmos/TransformSubGizmoUtil.cpp
@@ -0,0 +1,88 @@
+// Copyright Epic Games, Inc. All Rights Reserved.
+
+#include "BaseGizmos/TransformSubGizmoUtil.h"
+
+#include "BaseGizmos/CombinedTransformGizmo.h"
+#include "BaseGizmos/GizmoRenderingUtil.h"
+#include "BaseGizmos/GizmoPrivateUtil.h" // ToAxis, GetGizmoViewContext
+#include "BaseGizmos/ViewBasedTransformAdjusters.h"
+#include "BaseGizmos/ViewAdjustedStaticMeshGizmoComponent.h"
+#include "Components/SceneComponent.h"
+#include "Misc/EnumRange.h" // MakeFlagsRange
+
+
+FTransform UE::GizmoUtil::GetRotatedBasisTransform(const FTransform& TransformIn, EAxis::Type AxisToBeX)
+{
+	if (AxisToBeX == EAxis::Y)
+	{
+		FQuat RotationQuat(FMatrix(FVector::YAxisVector, FVector::ZAxisVector, FVector::XAxisVector, FVector::ZeroVector));
+		return TransformIn * RotationQuat;
+	}
+	else if (AxisToBeX == EAxis::Z)
+	{
+		FQuat RotationQuat(FMatrix(FVector::ZAxisVector, FVector::XAxisVector, FVector::YAxisVector, FVector::ZeroVector));
+		return TransformIn * RotationQuat;
+	}
+
+	return TransformIn;
+}
+
+void UE::GizmoUtil::InitializeSubGizmoElementsWithMesh(const FInitMeshSubGizmoParams& Params,
+	TArray<TPair<ETransformGizmoSubElements, UViewAdjustedStaticMeshGizmoComponent*>>* ComponentsOut)
+{
+	if (!ensure(Params.ParentGizmo && Params.Mesh))
+	{
+		return;
+	}
+
+	ACombinedTransformGizmoActor* GizmoActor = Params.ParentGizmo->GetGizmoActor();
+	UGizmoViewContext* GizmoViewContext = GetGizmoViewContext(Params.ParentGizmo->GetGizmoManager());
+
+	if (!ensure(GizmoActor && GizmoViewContext))
+	{
+		return;
+	}
+
+	// Iterate over each element individually
+	for (ETransformGizmoSubElements Element : MakeFlagsRange(Params.Elements))
+	{
+		EAxis::Type ElementAxis = ToAxis(Element);
+		FLinearColor Color = ElementAxis == EAxis::None ? FLinearColor::Black 
+			: UE::GizmoRenderingUtil::GetDefaultAxisColor(ElementAxis);
+
+		UViewAdjustedStaticMeshGizmoComponent* Component = UE::GizmoRenderingUtil::CreateDefaultMaterialGizmoMeshComponent(
+			Params.Mesh, GizmoViewContext, GizmoActor, Color);
+
+		// This makes it so that the component takes the place of the proper element
+		Params.ParentGizmo->SetSubGizmoComponent(Element, Component, FTransform::Identity);
+
+		// This has to happen after the SetSubGizmoComponent call above, so that it is already attached
+		Component->SetRelativeTransform(
+			Params.bRotateTransformBasisBasedOnElement ? GetRotatedBasisTransform(Params.ComponentToGizmo, ElementAxis)
+			: Params.ComponentToGizmo);
+
+		UE::GizmoRenderingUtil::FSubGizmoTransformAdjuster::AddTransformAdjuster(
+			Component, GizmoActor->GetRootComponent(), Params.bMirrorBasedOnOctant);
+
+		if (Params.SubstituteInteractionMesh)
+		{
+			UViewAdjustedStaticMeshGizmoComponent* SubstituteComponent = nullptr;
+			
+			SubstituteComponent = UE::GizmoRenderingUtil::CreateDefaultMaterialGizmoMeshComponent(
+				Params.SubstituteInteractionMesh, GizmoViewContext, Component, Color, 
+				// No need for hover material
+				false);
+
+			Component->SetSubstituteInteractionComponent(SubstituteComponent, Params.SubstituteMeshToComponent);
+
+			SubstituteComponent->SetRelativeTransform(Params.SubstituteMeshToComponent);
+			UE::GizmoRenderingUtil::FSubGizmoTransformAdjuster::AddTransformAdjuster(
+				SubstituteComponent, GizmoActor->GetRootComponent(), Params.bMirrorBasedOnOctant);
+		}//end substitute interaction mesh setup
+
+		if (ComponentsOut)
+		{
+			ComponentsOut->Emplace(Element, Component);
+		}
+	}//end for each element
+}
\ No newline at end of file
diff --git a/Engine/Source/Runtime/InteractiveToolsFramework/Public/BaseGizmos/AxisAngleGizmo.h b/Engine/Source/Runtime/InteractiveToolsFramework/Public/BaseGizmos/AxisAngleGizmo.h
--- a/Engine/Source/Runtime/InteractiveToolsFramework/Public/BaseGizmos/AxisAngleGizmo.h
+++ b/Engine/Source/Runtime/InteractiveToolsFramework/Public/BaseGizmos/AxisAngleGizmo.h
@@ -26,7 +26,11 @@ class IGizmoStateTarget;
 class UClickDragInputBehavior;
 class UObject;
 struct FToolBuilderState;
-
+namespace UE::GizmoUtil
+{
+	struct FTransformSubGizmoCommonParams;
+	struct FTransformSubGizmoSharedState;
+}
 
 UCLASS(MinimalAPI)
 class UAxisAngleGizmoBuilder : public UInteractiveGizmoBuilder
@@ -47,6 +51,14 @@ class UAxisAngleGizmo : public UInteractiveGizmo, public IClickDragBehaviorTarge
 	GENERATED_BODY()
 
 public:
+	/**
+	 * Helper that initializes AxisSource, AngleSource, HitTarget, and StateTarget for the common case 
+	 *  of being used to control the rotation of a transform. Safe to use for reinitialization.
+	 */
+	INTERACTIVETOOLSFRAMEWORK_API bool InitializeAsRotateGizmo(
+		const UE::GizmoUtil::FTransformSubGizmoCommonParams& InitializationParams,
+		UE::GizmoUtil::FTransformSubGizmoSharedState* SharedState);
+
 	// UInteractiveGizmo overrides
 
 	INTERACTIVETOOLSFRAMEWORK_API virtual void Setup() override;
@@ -69,6 +81,9 @@ public:
 	INTERACTIVETOOLSFRAMEWORK_API virtual void OnUpdateModifierState(int ModifierID, bool bIsOn) override;
 
 public:
+	// The below properties can be manipulated for more fine-grained control, but typically it is sufficient
+	// to use the initialization method above.
+
 	UPROPERTY()
 	TScriptInterface<IGizmoAxisSource> AxisSource;
 
diff --git a/Engine/Source/Runtime/InteractiveToolsFramework/Public/BaseGizmos/AxisPositionGizmo.h b/Engine/Source/Runtime/InteractiveToolsFramework/Public/BaseGizmos/AxisPositionGizmo.h
--- a/Engine/Source/Runtime/InteractiveToolsFramework/Public/BaseGizmos/AxisPositionGizmo.h
+++ b/Engine/Source/Runtime/InteractiveToolsFramework/Public/BaseGizmos/AxisPositionGizmo.h
@@ -27,6 +27,11 @@ class UClickDragInputBehavior;
 class UGizmoViewContext;
 class UObject;
 struct FToolBuilderState;
+namespace UE::GizmoUtil
+{
+	struct FTransformSubGizmoCommonParams;
+	struct FTransformSubGizmoSharedState;
+}
 
 UCLASS(MinimalAPI)
 class UAxisPositionGizmoBuilder : public UInteractiveGizmoBuilder
@@ -60,6 +65,22 @@ class UAxisPositionGizmo : public UInteractiveGizmo, public IClickDragBehaviorTa
 	GENERATED_BODY()
 
 public:
+	/**
+	 * Helper that initializes AxisSource, ParameterSource, HitTarget, and StateTarget for the common case
+	 *  of being used to control the translation of a transform. Safe to use for reinitialization.
+	 */
+	INTERACTIVETOOLSFRAMEWORK_API bool InitializeAsTranslateGizmo(
+		const UE::GizmoUtil::FTransformSubGizmoCommonParams& InitializationParams,
+		UE::GizmoUtil::FTransformSubGizmoSharedState* SharedState);
+	/**
+	 * Helper that initializes AxisSource, ParameterSource, HitTarget, and StateTarget for the common case
+	 *  of being used to control the scale of a transform. Safe to use for reinitialization.
+	 */
+	INTERACTIVETOOLSFRAMEWORK_API bool InitializeAsScaleGizmo(
+		const UE::GizmoUtil::FTransformSubGizmoCommonParams& InitializationParams,
+		bool bDisallowNegativeScaling,
+		UE::GizmoUtil::FTransformSubGizmoSharedState* SharedState);
+
 	// UInteractiveGizmo overrides
 
 	INTERACTIVETOOLSFRAMEWORK_API virtual void Setup() override;
@@ -80,6 +101,10 @@ public:
 
 
 public:
+
+	// The below properties can be manipulated for more fine-grained control, but typically it is sufficient
+	// to use one of the initialization methods above.
+
 	/** AxisSource provides the 3D line on which the interaction happens */
 	UPROPERTY()
 	TScriptInterface<IGizmoAxisSource> AxisSource;
@@ -88,6 +113,7 @@ public:
 	UPROPERTY()
 	TScriptInterface<IGizmoFloatParameterSource> ParameterSource;
 
+	UE_DEPRECATED(5.5, "UAxisPositionGizmo does not use GizmoViewContext itself (it is used in gizmo components for hit testing).")
 	UPROPERTY()
 	TObjectPtr<UGizmoViewContext> GizmoViewContext;
 
diff --git a/Engine/Source/Runtime/InteractiveToolsFramework/Public/BaseGizmos/CombinedTransformGizmo.h b/Engine/Source/Runtime/InteractiveToolsFramework/Public/BaseGizmos/CombinedTransformGizmo.h
--- a/Engine/Source/Runtime/InteractiveToolsFramework/Public/BaseGizmos/CombinedTransformGizmo.h
+++ b/Engine/Source/Runtime/InteractiveToolsFramework/Public/BaseGizmos/CombinedTransformGizmo.h
@@ -21,6 +21,11 @@ class UGizmoComponentAxisSource;
 class UGizmoTransformChangeStateTarget;
 class UGizmoViewContext;
 class FTransformGizmoTransformChange;
+namespace UE::GizmoUtil
+{
+	struct FTransformSubGizmoCommonParams;
+	struct FTransformSubGizmoSharedState;
+}
 
 /**
  * ACombinedTransformGizmoActor is an Actor type intended to be used with UCombinedTransformGizmo,
@@ -89,7 +94,7 @@ public:
 	UPROPERTY()
 	TObjectPtr<UPrimitiveComponent> RotateZ;
 
-	/** Z Axis Rotation Component */
+	/** Circle that gets drawn around the outside of the gizmo to make it look like a sphere */
 	UPROPERTY()
 	TObjectPtr<UPrimitiveComponent> RotationSphere;
 
@@ -130,6 +135,22 @@ public:
 
 
 public:
+	/**
+	 * Replaces the component corresponding to the given element with a new component.
+	 * 
+	 * @param Element Element to replace, should be a single element (no combined flags), except that
+	 *  RotateAllAxes will be interpreted to mean the rotation sphere component.
+	 * @param NewComponent The component to replace with. If nullptr, no component is added (i.e. the
+	 *  function just deletes the existing component). If not nullptr, then the component should have
+	 *  this actor in its outer chain, i.e. it should have been created as NewObject<UYourComponentType>(ThisActor).
+	 * @param SubGizmoToGizmo Transform from component to gizmo (i.e. the relative transform).
+	 * @param ReplacedComponentOut Outputs a pointer to the replaced component, if there was one. Note 
+	 *  that the component will already have had DestroyComponent() called on it.
+	 * @return true if successful, which will be always as long as parameters are valid.
+	 */
+	bool ReplaceSubGizmoComponent(ETransformGizmoSubElements Element, UPrimitiveComponent* NewComponent, 
+		const FTransform& SubGizmoToGizmo, UPrimitiveComponent** ReplacedComponentOut = nullptr);
+
 	/**
 	 * Create a new instance of ACombinedTransformGizmoActor and populate the various
 	 * sub-components with standard GizmoXComponent instances suitable for a 3-axis transformer Gizmo
@@ -301,6 +322,21 @@ public:
 	INTERACTIVETOOLSFRAMEWORK_API virtual void SetUpdateHoverFunction(TFunction<void(UPrimitiveComponent*, bool)> HoverFunction);
 	INTERACTIVETOOLSFRAMEWORK_API virtual void SetUpdateCoordSystemFunction(TFunction<void(UPrimitiveComponent*, EToolContextCoordinateSystem)> CoordSysFunction);
 	
+	/**
+	 * Sets a given sub gizmo component to the given component. This is only valid to call after Setup(), but
+	 *  can be before or after SetActiveTarget.
+	 * 
+	 * @param Element Element to replace, should be a single element (no combined flags), except that
+	 *  RotateAllAxes will be interpreted to mean the rotation sphere component.
+	 * @param NewComponent The component to replace with. If nullptr, no component is added (i.e. the
+	 *   function just deletes the existing component). If not nullptr, then the component should have
+	 *   the gizmo actor in its outer chain, e.g. NewObject<UYourComponentType>(ThisGizmo->GetGizmoActor()).
+	 * @param SubGizmoToGizmo Transform from component to gizmo.
+	 * @return true if successful.
+	 */
+	INTERACTIVETOOLSFRAMEWORK_API bool SetSubGizmoComponent(ETransformGizmoSubElements Element, UPrimitiveComponent* Component, 
+		const FTransform& SubGizmoToGizmo);
+
 	/**
 	 * If used, binds alignment functions to the sub gizmos that they can use to align to geometry in the scene. 
 	 * Specifically, translation and rotation gizmos will check ShouldAlignDestination() to see if they should
@@ -646,17 +682,52 @@ protected:
 
 	// See comment for SetDisplaySpaceTransform;
 	TOptional<FTransform> DisplaySpaceTransform;
+
 protected:
 
+	using FTransformSubGizmoCommonParams = UE::GizmoUtil::FTransformSubGizmoCommonParams;
+	using FTransformSubGizmoSharedState = UE::GizmoUtil::FTransformSubGizmoSharedState;
+
+	/** @return a new instance of the standard axis-translation Gizmo */
+	INTERACTIVETOOLSFRAMEWORK_API virtual UInteractiveGizmo* AddAxisTranslationGizmo(
+		FTransformSubGizmoCommonParams& Params, FTransformSubGizmoSharedState& SharedState);
+
+	/** @return a new instance of the standard plane-translation Gizmo */
+	INTERACTIVETOOLSFRAMEWORK_API virtual UInteractiveGizmo* AddPlaneTranslationGizmo(
+		FTransformSubGizmoCommonParams& Params, FTransformSubGizmoSharedState& SharedState);
+
+	/** @return a new instance of the standard axis-rotation Gizmo */
+	INTERACTIVETOOLSFRAMEWORK_API virtual UInteractiveGizmo* AddAxisRotationGizmo(
+		FTransformSubGizmoCommonParams& Params, FTransformSubGizmoSharedState& SharedState);
+
+	/** @return a new instance of the standard axis-scaling Gizmo */
+	INTERACTIVETOOLSFRAMEWORK_API virtual UInteractiveGizmo* AddAxisScaleGizmo(
+		FTransformSubGizmoCommonParams& Params, FTransformSubGizmoSharedState& SharedState);
+
+	/** @return a new instance of the standard plane-scaling Gizmo */
+	INTERACTIVETOOLSFRAMEWORK_API virtual UInteractiveGizmo* AddPlaneScaleGizmo(
+		FTransformSubGizmoCommonParams& Params, FTransformSubGizmoSharedState& SharedState);
+
+	/** @return a new instance of the standard plane-scaling Gizmo */
+	INTERACTIVETOOLSFRAMEWORK_API virtual UInteractiveGizmo* AddUniformScaleGizmo(
+		FTransformSubGizmoCommonParams& Params, FTransformSubGizmoSharedState& SharedState);
+
+	// Axis and Plane TransformSources use these function to execute snapping queries
+	INTERACTIVETOOLSFRAMEWORK_API bool PositionSnapFunction(const FVector& WorldPosition, FVector& SnappedPositionOut) const;
+	INTERACTIVETOOLSFRAMEWORK_API bool PositionAxisDeltaSnapFunction(double AxisDelta, double& SnappedDeltaOut, int AxisIndex) const;
+	INTERACTIVETOOLSFRAMEWORK_API FQuat RotationSnapFunction(const FQuat& DeltaRotation) const;
+	INTERACTIVETOOLSFRAMEWORK_API bool RotationAxisAngleSnapFunction(double AxisAngleDelta, double& SnappedAxisAngleDeltaOut, int AxisIndex) const;
+
 
+	UE_DEPRECATED(5.5, "Use FTransformSubGizmoCommonParams overload instead.")
 	/** @return a new instance of the standard axis-translation Gizmo */
 	INTERACTIVETOOLSFRAMEWORK_API virtual UInteractiveGizmo* AddAxisTranslationGizmo(
 		UPrimitiveComponent* AxisComponent, USceneComponent* RootComponent,
 		IGizmoAxisSource* AxisSource,
-		IGizmoTransformSource* TransformSource, 
+		IGizmoTransformSource* TransformSource,
 		IGizmoStateTarget* StateTarget,
 		int AxisIndex);
-
+	UE_DEPRECATED(5.5, "Use FTransformSubGizmoCommonParams overload instead.")
 	/** @return a new instance of the standard plane-translation Gizmo */
 	INTERACTIVETOOLSFRAMEWORK_API virtual UInteractiveGizmo* AddPlaneTranslationGizmo(
 		UPrimitiveComponent* AxisComponent, USceneComponent* RootComponent,
@@ -664,28 +735,28 @@ protected:
 		IGizmoTransformSource* TransformSource,
 		IGizmoStateTarget* StateTarget,
 		int XAxisIndex, int YAxisIndex);
-
+	UE_DEPRECATED(5.5, "Use FTransformSubGizmoCommonParams overload instead.")
 	/** @return a new instance of the standard axis-rotation Gizmo */
 	INTERACTIVETOOLSFRAMEWORK_API virtual UInteractiveGizmo* AddAxisRotationGizmo(
 		UPrimitiveComponent* AxisComponent, USceneComponent* RootComponent,
 		IGizmoAxisSource* AxisSource,
 		IGizmoTransformSource* TransformSource,
 		IGizmoStateTarget* StateTarget);
-
+	UE_DEPRECATED(5.5, "Use FTransformSubGizmoCommonParams overload instead.")
 	/** @return a new instance of the standard axis-scaling Gizmo */
 	INTERACTIVETOOLSFRAMEWORK_API virtual UInteractiveGizmo* AddAxisScaleGizmo(
 		UPrimitiveComponent* AxisComponent, USceneComponent* RootComponent,
 		IGizmoAxisSource* GizmoAxisSource, IGizmoAxisSource* ParameterAxisSource,
 		IGizmoTransformSource* TransformSource,
 		IGizmoStateTarget* StateTarget);
-
+	UE_DEPRECATED(5.5, "Use FTransformSubGizmoCommonParams overload instead.")
 	/** @return a new instance of the standard plane-scaling Gizmo */
 	INTERACTIVETOOLSFRAMEWORK_API virtual UInteractiveGizmo* AddPlaneScaleGizmo(
 		UPrimitiveComponent* AxisComponent, USceneComponent* RootComponent,
 		IGizmoAxisSource* GizmoAxisSource, IGizmoAxisSource* ParameterAxisSource,
 		IGizmoTransformSource* TransformSource,
 		IGizmoStateTarget* StateTarget);
-
+	UE_DEPRECATED(5.5, "Use FTransformSubGizmoCommonParams overload instead.")
 	/** @return a new instance of the standard plane-scaling Gizmo */
 	INTERACTIVETOOLSFRAMEWORK_API virtual UInteractiveGizmo* AddUniformScaleGizmo(
 		UPrimitiveComponent* ScaleComponent, USceneComponent* RootComponent,
@@ -693,10 +764,10 @@ protected:
 		IGizmoTransformSource* TransformSource,
 		IGizmoStateTarget* StateTarget);
 
-	// Axis and Plane TransformSources use these function to execute snapping queries
-	INTERACTIVETOOLSFRAMEWORK_API bool PositionSnapFunction(const FVector& WorldPosition, FVector& SnappedPositionOut) const;
-	INTERACTIVETOOLSFRAMEWORK_API bool PositionAxisDeltaSnapFunction(double AxisDelta, double& SnappedDeltaOut, int AxisIndex) const;
-	INTERACTIVETOOLSFRAMEWORK_API FQuat RotationSnapFunction(const FQuat& DeltaRotation) const;
-	INTERACTIVETOOLSFRAMEWORK_API bool RotationAxisAngleSnapFunction(double AxisAngleDelta, double& SnappedAxisAngleDeltaOut, int AxisIndex) const;
-
+	// Useful for reinitializing components after SetActiveTarget, or for use by derived classes.
+	TUniquePtr<FTransformSubGizmoSharedState> SubGizmoSharedState;
+private:
+	// Here to support subgizmo reinitialization after SetActiveTarget has been called. Private
+	// instead of protected for now in case we change the approach here.
+	IToolContextTransactionProvider* TransactionProviderAtLastSetActiveTarget = nullptr;
 };
diff --git a/Engine/Source/Runtime/InteractiveToolsFramework/Public/BaseGizmos/PlanePositionGizmo.h b/Engine/Source/Runtime/InteractiveToolsFramework/Public/BaseGizmos/PlanePositionGizmo.h
--- a/Engine/Source/Runtime/InteractiveToolsFramework/Public/BaseGizmos/PlanePositionGizmo.h
+++ b/Engine/Source/Runtime/InteractiveToolsFramework/Public/BaseGizmos/PlanePositionGizmo.h
@@ -27,7 +27,11 @@ class IGizmoVec2ParameterSource;
 class UClickDragInputBehavior;
 class UObject;
 struct FToolBuilderState;
-
+namespace UE::GizmoUtil
+{
+	struct FTransformSubGizmoCommonParams;
+	struct FTransformSubGizmoSharedState;
+}
 
 UCLASS(MinimalAPI)
 class UPlanePositionGizmoBuilder : public UInteractiveGizmoBuilder
@@ -62,9 +66,40 @@ class UPlanePositionGizmo : public UInteractiveGizmo, public IClickDragBehaviorT
 	GENERATED_BODY()
 
 public:
+	/**
+	 * Helper that initializes AxisSource, ParameterSource, HitTarget, and StateTarget for the common case
+	 *  of being used to control the translation of a transform. Safe to use for reinitialization.
+	 */
+	INTERACTIVETOOLSFRAMEWORK_API bool InitializeAsTranslateGizmo(
+		const UE::GizmoUtil::FTransformSubGizmoCommonParams& InitializationParams,
+		UE::GizmoUtil::FTransformSubGizmoSharedState* SharedState);
+	
+	/**
+	 * Helper that initializes AxisSource, ParameterSource, HitTarget, and StateTarget for the common case
+	 *  of being used to control the scale of a transform. Safe to use for reinitialization.
+	 */
+	INTERACTIVETOOLSFRAMEWORK_API bool InitializeAsScaleGizmo(
+		const UE::GizmoUtil::FTransformSubGizmoCommonParams& InitializationParams,
+		bool bDisallowNegativeScaling,
+		UE::GizmoUtil::FTransformSubGizmoSharedState* SharedState);
+
+	/**
+	 * Helper that initializes AxisSource, ParameterSource, HitTarget, and StateTarget for the common case
+	 *  of being used to control the uniform scale of a transform, where the interaction plane remains aligned
+	 *  to the camera. Safe to use for reinitialization.
+	 * Note: if SharedState contains a CameraAxisSource, this subgizmo will use it and will assume that something
+	 *  else is responsible for updating it. Otherwise this subgizmo will set the CameraAxisSource and update it
+	 *  in its own Tick function.
+	 */
+	bool InitializeAsUniformScaleGizmo(
+		const UE::GizmoUtil::FTransformSubGizmoCommonParams& Params,
+		bool bDisallowNegativeScaling,
+		UE::GizmoUtil::FTransformSubGizmoSharedState* SharedState);
+	
 	// UInteractiveGizmo overrides
 
 	INTERACTIVETOOLSFRAMEWORK_API virtual void Setup() override;
+	INTERACTIVETOOLSFRAMEWORK_API virtual void Tick(float DeltaTime) override;
 
 	// IClickDragBehaviorTarget implementation
 
@@ -81,6 +116,9 @@ public:
 	INTERACTIVETOOLSFRAMEWORK_API virtual void OnEndHover() override;
 
 public:
+	// The below properties can be manipulated for more fine-grained control, but typically it is sufficient
+	// to use one of the initialization methods above.
+
 	/** AxisSource provides the 3D plane on which the interaction happens */
 	UPROPERTY()
 	TScriptInterface<IGizmoAxisSource> AxisSource;
@@ -183,5 +221,8 @@ protected:
 
 	FVector2D InteractionStartOriginParameterOffset;
 	FInputDeviceRay LastInputRay = FInputDeviceRay(FRay());
+
+private:
+	TUniqueFunction<void(float DeltaTime)> CustomTickFunction;
 };
 
diff --git a/Engine/Source/Runtime/InteractiveToolsFramework/Public/BaseGizmos/TransformSubGizmoUtil.h b/Engine/Source/Runtime/InteractiveToolsFramework/Public/BaseGizmos/TransformSubGizmoUtil.h
new file mode 100644
--- /dev/null
+++ b/Engine/Source/Runtime/InteractiveToolsFramework/Public/BaseGizmos/TransformSubGizmoUtil.h
@@ -0,0 +1,158 @@
+// Copyright Epic Games, Inc. All Rights Reserved.
+
+#pragma once
+
+#include "Containers/Array.h"
+#include "InteractiveGizmo.h"
+#include "Math/Axis.h"
+#include "Math/MathFwd.h" // FTransform
+#include "Misc/Optional.h"
+
+struct FLinearColor;
+class FText;
+class IToolContextTransactionProvider;
+class UCombinedTransformGizmo;
+class UGizmoScaledAndUnscaledTransformSources;
+class UGizmoTransformChangeStateTarget;
+class UGizmoConstantFrameAxisSource;
+class UGizmoComponentAxisSource;
+class UObject;
+class UPrimitiveComponent;
+class UStaticMesh;
+class UTransformProxy;
+class UViewAdjustedStaticMeshGizmoComponent;
+
+/**
+ * This file is meant for utilities that are specific to allowing gizmos to act as sub gizmos inside a
+ *  UCombinedTransformGizmo.
+ */
+
+namespace UE::GizmoUtil
+{
+	/**
+	 * Common parameters needed for initializing a sub gizmo. Used in several sub gizmo reinitialization
+	 *  functions (such as UAxisPositionGizmo::InitializeAsTranslateGizmo).
+	 */
+	struct FTransformSubGizmoCommonParams
+	{
+		UPrimitiveComponent* Component = nullptr;
+		UTransformProxy* TransformProxy = nullptr;
+		EAxis::Type Axis = EAxis::None;
+
+		// If true, we use the axis of the root gizmo component. If false, we use an axis of the
+		//  component itself, which is determined by its transform. Does not attempt to use shared
+		//  state for axis if false.
+		bool bAxisIsBasedOnRootComponent = true;
+
+		// Optional parameters:
+
+		// Default transaction name and provider are used if not provided. Not used if one is gotten from
+		//  shared state.
+		TOptional<FText> TransactionName;
+		IToolContextTransactionProvider* TransactionProvider = nullptr;
+
+		// Transient package is used if custom outer is not provided
+		UObject* OuterForSubobjects = nullptr;
+		
+		// If true, the gizmo moves the parent root component rather than just its own component, like
+		//  the standard TRS sub gizmos do. If false, the gizmo moves just its component, leaving the
+		//  parent in place. Does not attempt to use shared state for transform source or state target
+		//  if false.
+		bool bManipulatesRootComponent = true;
+		
+		//~ The default hover behavior is to forward the hover information to the component via the
+		//~  IGizmoBaseComponentInterface. We could give the user a way to place a custom hover function
+		//~  like we have in deprecated initialization methods, but that is unclear to be worth it, as 
+		//~  it is likely to be cleaner to implement the interface on any custom components.
+		//~TFunction<void(bool)> CustomUpdateHoverFunction;
+
+		// Returns the Axis parameter as an index in the range [0,2] where 0 
+		// corresponds to X (or None), 1 corresponds to Y, and 2 corresponds to Z.
+		int GetClampedAxisIndex() const
+		{
+			return (Axis == EAxis::Y) ? 1
+				: (Axis == EAxis::Z) ? 2
+				: 0;
+		}
+	};
+
+	/**
+	 * A struct to hold some objects that can be reused across multiple sub gizmos that operate on the same
+	 *  compound gizmo, to avoid creating redundant objects. For instance, the transform source can be the
+	 *  same for the sub gizmos that manipulate the same overall TRS gizmo, so the first sub gizmo would
+	 *  create that transform source, and the subsequent ones will reuse it.
+	 * 
+	 * Note: While the initialization functions will try to do reasonable things with this shared state
+	 *  (for instance, not attempt to use shared state for the axis source if the axis source is marked
+	 *  as being based off of this component, rather than the root), it is still up to the user to decide
+	 *  whether shared state should be an option. For instance, if your sub gizmo manipulates the root as
+	 *  normal but for some reason targets a different TransformProxy than other subgizmos, then it
+	 *  shouldn't share a TransformSource or StateTarget with the other gizmos, so either the shared state
+	 *  struct should not be used (preferable), or those properties in the struct should be nulled out
+	 *  before/after the call.
+	 * In general, if your sub gizmo is part of the overall TRS gizmo itself (i.e. represents one of its elements),
+	 *  then using the shared state struct will make sense, and will save you a few redundant object creations.
+	 *  If your sub gizmo is just attached to the TRS gizmo but moves independently, then it may not be safe to
+	 *  use the shared state struct, and may not be worth the potential trouble even if you could save an object.
+	 */
+	struct FTransformSubGizmoSharedState
+	{
+		UGizmoScaledAndUnscaledTransformSources* TransformSource = nullptr;
+		UGizmoTransformChangeStateTarget* StateTarget = nullptr;
+		UGizmoConstantFrameAxisSource* CameraAxisSource = nullptr;
+		UGizmoComponentAxisSource* CardinalAxisSources[3] = { 0 };
+		UGizmoComponentAxisSource* UnitCardinalAxisSources[3] = { 0 };
+	};
+
+	/**
+	 * Rotates a transform such that its basis still lies along the cardinal axes but rotated in such a 
+	 *  way that the given axis is in the X direction. For example, if you have a transform that makes a
+	 *  component work well as an X axis in a TRS gizmo, using this function with AxisToBeX set to Y will
+	 *  rotate the transform to work well as the Y axis, because the basis will have been rotated to be
+	 *  YZX.
+	 */
+	INTERACTIVETOOLSFRAMEWORK_API FTransform GetRotatedBasisTransform(const FTransform& TransformIn, EAxis::Type AxisToBeX);
+
+	// Parameter struct for InitializeSubGizmoElementsWithMesh
+	struct FInitMeshSubGizmoParams
+	{
+		// Required parameters:
+
+		UCombinedTransformGizmo* ParentGizmo = nullptr;
+		// If there are multiple elements, the same mesh will be set for all of them. This makes
+		//  it easy to, for example, set all of the arrow components at once.
+		ETransformGizmoSubElements Elements = ETransformGizmoSubElements::None;
+		UStaticMesh* Mesh = nullptr;
+		
+		// Optional parameters:
+
+		FTransform ComponentToGizmo = FTransform::Identity;
+
+		// When true, ComponentToGizmo is adjusted such that the gizmo basis points along
+		//  the relevant element axis. This allows ComponentToGizmo to be specified once
+		//  for the x axis case, and be reused for the y/z axes with the proper rotation.
+		bool bRotateTransformBasisBasedOnElement = true;
+
+		// When true, the sub gizmo is mirrored across the gizmo origin depending on where
+		//  in relation to the gizmo the camera is positioned.
+		bool bMirrorBasedOnOctant = false;
+
+		// Leaving this unset causes the color to be determined by axis.
+		TOptional<FLinearColor> Color;
+
+		// Mesh to swap in when the user is dragging the gizmo.
+		UStaticMesh* SubstituteInteractionMesh = nullptr;
+		// Only used when SubstituteInteractionMesh is being set
+		FTransform SubstituteMeshToComponent = FTransform::Identity;
+	};
+
+	/**
+	 * Given a UCombinedTransformGizmo, swaps selected elements with custom meshes.
+	 * 
+	 * @param ComponentsOut Optional output array of newly created gizmo components, as pairs with the
+	 *  single element that they were created for.
+	 */
+	INTERACTIVETOOLSFRAMEWORK_API void InitializeSubGizmoElementsWithMesh(const FInitMeshSubGizmoParams& Params, 
+		TArray<TPair<ETransformGizmoSubElements, UViewAdjustedStaticMeshGizmoComponent*>>* ElementComponentsOut = nullptr);
+
+}
