InteractiveToolsFramework: Add simple single click gizmo and a helper function for building gizmos without having to register them ahead of time.

diff --git a/Engine/Source/Runtime/InteractiveToolsFramework/Private/BaseGizmos/GizmoUtil.cpp b/Engine/Source/Runtime/InteractiveToolsFramework/Private/BaseGizmos/GizmoUtil.cpp
new file mode 100644
--- /dev/null
+++ b/Engine/Source/Runtime/InteractiveToolsFramework/Private/BaseGizmos/GizmoUtil.cpp
@@ -0,0 +1,26 @@
+// Copyright Epic Games, Inc. All Rights Reserved.
+
+#include "BaseGizmos/GizmoUtil.h"
+
+#include "InteractiveGizmoManager.h"
+#include "InteractiveGizmo.h"
+
+UInteractiveGizmo* UE::GizmoUtil::CreateGizmoViaSimpleBuilder(UInteractiveGizmoManager* GizmoManager, 
+	TSubclassOf<UInteractiveGizmo> GizmoClass, const FString& InstanceIdentifier, void* Owner)
+{
+	if (!ensure(GizmoManager))
+	{
+		return nullptr;
+	}
+
+	const FString BuilderIdentifier = TEXT("__CreateGizmoViaSimpleBuilder_TemporaryBuilder");
+
+	USimpleLambdaInteractiveGizmoBuilder* Builder = NewObject<USimpleLambdaInteractiveGizmoBuilder>();
+	Builder->BuilderFunc = [GizmoManager, GizmoClass](const FToolBuilderState& SceneState) { return NewObject<UInteractiveGizmo>(GizmoManager, GizmoClass.Get()); };
+
+	GizmoManager->RegisterGizmoType(BuilderIdentifier, Builder);
+	UInteractiveGizmo* Gizmo = GizmoManager->CreateGizmo(BuilderIdentifier, InstanceIdentifier, Owner);
+	GizmoManager->DeregisterGizmoType(BuilderIdentifier);
+
+	return Gizmo;
+}
diff --git a/Engine/Source/Runtime/InteractiveToolsFramework/Private/BaseGizmos/SimpleSingleClickGizmo.cpp b/Engine/Source/Runtime/InteractiveToolsFramework/Private/BaseGizmos/SimpleSingleClickGizmo.cpp
new file mode 100644
--- /dev/null
+++ b/Engine/Source/Runtime/InteractiveToolsFramework/Private/BaseGizmos/SimpleSingleClickGizmo.cpp
@@ -0,0 +1,80 @@
+// Copyright Epic Games, Inc. All Rights Reserved.
+
+#include "BaseGizmos/SimpleSingleClickGizmo.h"
+
+#include "BaseBehaviors/MouseHoverBehavior.h"
+#include "BaseBehaviors/SingleClickBehavior.h"
+#include "BaseGizmos/GizmoBaseComponent.h" // IGizmoBaseComponentInterface
+#include "BaseGizmos/HitTargets.h"
+
+bool USimpleSingleClickGizmo::InitializeWithComponent(UPrimitiveComponent* ComponentIn)
+{
+	if (!ComponentIn)
+	{
+		return false;
+	}
+
+	UGizmoComponentHitTarget* CastHitTarget = UGizmoComponentHitTarget::Construct(ComponentIn, this);
+	CastHitTarget->UpdateHoverFunction = [ComponentIn](bool bHovering)
+	{
+		if (IGizmoBaseComponentInterface* Hoverable = Cast<IGizmoBaseComponentInterface>(ComponentIn))
+		{
+			Hoverable->UpdateHoverState(bHovering);
+		}
+	};
+	HitTarget = CastHitTarget;
+	return true;
+}
+
+void USimpleSingleClickGizmo::Setup()
+{
+	ClickBehavior = NewObject<USingleClickInputBehavior>(this);
+	ClickBehavior->Initialize(this);
+	ClickBehavior->SetDefaultPriority(FInputCapturePriority(FInputCapturePriority::DEFAULT_GIZMO_PRIORITY));
+	AddInputBehavior(ClickBehavior);
+	
+	UMouseHoverBehavior* HoverBehavior = NewObject<UMouseHoverBehavior>(this);
+	HoverBehavior->Initialize(this);
+	HoverBehavior->SetDefaultPriority(FInputCapturePriority(FInputCapturePriority::DEFAULT_GIZMO_PRIORITY));
+	AddInputBehavior(HoverBehavior);
+}
+
+FInputRayHit USimpleSingleClickGizmo::IsHitByClick(const FInputDeviceRay& ClickPos)
+{
+	FInputRayHit GizmoHit;
+	if (HitTarget)
+	{
+		GizmoHit = HitTarget->IsHit(ClickPos);
+	}
+	return GizmoHit;
+}
+
+void USimpleSingleClickGizmo::OnClicked(const FInputDeviceRay& ClickPos)
+{
+	OnClick.Broadcast(*this, ClickPos);
+}
+
+FInputRayHit USimpleSingleClickGizmo::BeginHoverSequenceHitTest(const FInputDeviceRay& PressPos)
+{
+	FInputRayHit GizmoHit;
+	if (HitTarget)
+	{
+		GizmoHit = HitTarget->IsHit(PressPos);
+	}
+	return GizmoHit;
+}
+
+void USimpleSingleClickGizmo::OnBeginHover(const FInputDeviceRay& DevicePos)
+{
+	HitTarget->UpdateHoverState(true);
+}
+
+bool USimpleSingleClickGizmo::OnUpdateHover(const FInputDeviceRay& DevicePos)
+{
+	return true;
+}
+
+void USimpleSingleClickGizmo::OnEndHover()
+{
+	HitTarget->UpdateHoverState(false);
+}
\ No newline at end of file
diff --git a/Engine/Source/Runtime/InteractiveToolsFramework/Public/BaseGizmos/GizmoUtil.h b/Engine/Source/Runtime/InteractiveToolsFramework/Public/BaseGizmos/GizmoUtil.h
new file mode 100644
--- /dev/null
+++ b/Engine/Source/Runtime/InteractiveToolsFramework/Public/BaseGizmos/GizmoUtil.h
@@ -0,0 +1,60 @@
+// Copyright Epic Games, Inc. All Rights Reserved.
+
+#pragma once
+
+#include "InteractiveGizmoBuilder.h"
+#include "Templates/Function.h"
+#include "Templates/SubclassOf.h"
+
+#include "GizmoUtil.generated.h"
+
+class FString;
+class UInteractiveGizmo;
+class UInteractiveGizmoManager;
+
+/**
+ * Gizmo builder that simply calls a particular lambda when building a gizmo. Makes it easy to
+ *  register gizmo build behavior without writing a new class.
+ */
+UCLASS(MinimalAPI)
+class USimpleLambdaInteractiveGizmoBuilder : public UInteractiveGizmoBuilder
+{
+	GENERATED_BODY()
+
+public:
+
+	virtual UInteractiveGizmo* BuildGizmo(const FToolBuilderState& SceneState) const override
+	{
+		if (BuilderFunc)
+		{
+			return BuilderFunc(SceneState);
+		}
+		return nullptr;
+	}
+
+	TUniqueFunction<UInteractiveGizmo* (const FToolBuilderState& SceneState)> BuilderFunc;
+};
+
+namespace UE::GizmoUtil
+{
+	/**
+	 * Uses the gizmo manager to create a gizmo of the given class (assuming that the gizmo type does not need
+	 *  any special setup beyond instantiation) without having to register a custom builder for that class ahead of time.
+	 * 
+	 * This function lets the user bypass the need to define, register, and use a builder class, while still registering
+	 *  the gizmo properly with the gizmo manager. Under the hood, it creates and registers a temporary generic builder, 
+	 *  uses it to make the gizmo, and then immediately deregisters the builder. 
+	 */
+	INTERACTIVETOOLSFRAMEWORK_API UInteractiveGizmo* CreateGizmoViaSimpleBuilder(UInteractiveGizmoManager* GizmoManager,
+		TSubclassOf<UInteractiveGizmo> GizmoClass, const FString& InstanceIdentifier, void* Owner);
+
+	/**
+	 * Template version of CreateGizmoViaSimpleBuilder that does a cast on return.
+	 */
+	template <typename GizmoClass>
+	GizmoClass* CreateGizmoViaSimpleBuilder(UInteractiveGizmoManager* GizmoManager,
+		const FString& InstanceIdentifier, void* GizmoOwner)
+	{
+		return Cast<GizmoClass>(CreateGizmoViaSimpleBuilder(GizmoManager, GizmoClass::StaticClass(), InstanceIdentifier, GizmoOwner));
+	}
+}
\ No newline at end of file
diff --git a/Engine/Source/Runtime/InteractiveToolsFramework/Public/BaseGizmos/SimpleSingleClickGizmo.h b/Engine/Source/Runtime/InteractiveToolsFramework/Public/BaseGizmos/SimpleSingleClickGizmo.h
new file mode 100644
--- /dev/null
+++ b/Engine/Source/Runtime/InteractiveToolsFramework/Public/BaseGizmos/SimpleSingleClickGizmo.h
@@ -0,0 +1,52 @@
+// Copyright Epic Games, Inc. All Rights Reserved.
+
+#pragma once
+
+#include "BaseBehaviors/BehaviorTargetInterfaces.h" // IClickBehaviorTarget, IHoverBehaviorTarget
+#include "InteractiveGizmo.h"
+
+#include "SimpleSingleClickGizmo.generated.h"
+
+class IGizmoClickTarget;
+class UPrimitiveComponent;
+class USingleClickInputBehavior;
+
+/**
+ * Simple gizmo that triggers an OnClicked callback when it is clicked.
+ */
+UCLASS(MinimalAPI)
+class USimpleSingleClickGizmo : public UInteractiveGizmo, public IClickBehaviorTarget, public IHoverBehaviorTarget
+{
+	GENERATED_BODY()
+
+public:
+	
+	INTERACTIVETOOLSFRAMEWORK_API bool InitializeWithComponent(UPrimitiveComponent* ComponentIn);
+	
+	// UInteractiveGizmo
+	INTERACTIVETOOLSFRAMEWORK_API virtual void Setup() override;
+
+	// IClickBehaviorTarget
+	INTERACTIVETOOLSFRAMEWORK_API virtual FInputRayHit IsHitByClick(const FInputDeviceRay& ClickPos) override;
+	INTERACTIVETOOLSFRAMEWORK_API virtual void OnClicked(const FInputDeviceRay& ClickPos) override;
+
+	// IHoverBehaviorTarget
+	INTERACTIVETOOLSFRAMEWORK_API virtual FInputRayHit BeginHoverSequenceHitTest(const FInputDeviceRay& PressPos) override;
+	INTERACTIVETOOLSFRAMEWORK_API virtual void OnBeginHover(const FInputDeviceRay& DevicePos) override;
+	INTERACTIVETOOLSFRAMEWORK_API virtual bool OnUpdateHover(const FInputDeviceRay& DevicePos) override;
+	INTERACTIVETOOLSFRAMEWORK_API virtual void OnEndHover() override;
+
+public:
+
+	DECLARE_MULTICAST_DELEGATE_TwoParams(FOnClick, USimpleSingleClickGizmo&, const FInputDeviceRay&);
+	FOnClick OnClick;
+
+	/** The HitTarget provides a hit-test against some 3D element (presumably a visual widget) that controls when interaction can start */
+	UPROPERTY()
+	TScriptInterface<IGizmoClickTarget> HitTarget;
+
+	/** The mouse click behavior of the gizmo is accessible so that it can be modified to use different mouse keys. */
+	UPROPERTY()
+	TObjectPtr<USingleClickInputBehavior> ClickBehavior;
+
+};
\ No newline at end of file
