InteractiveToolsFramework: Add missing include to fix cis.

diff --git a/Engine/Source/Runtime/InteractiveToolsFramework/Private/BaseGizmos/GizmoLineHandleComponent.cpp b/Engine/Source/Runtime/InteractiveToolsFramework/Private/BaseGizmos/GizmoLineHandleComponent.cpp
--- a/Engine/Source/Runtime/InteractiveToolsFramework/Private/BaseGizmos/GizmoLineHandleComponent.cpp
+++ b/Engine/Source/Runtime/InteractiveToolsFramework/Private/BaseGizmos/GizmoLineHandleComponent.cpp
@@ -3,6 +3,7 @@
 #include "BaseGizmos/GizmoLineHandleComponent.h"
 #include "BaseGizmos/GizmoRenderingUtil.h"
 #include "BaseGizmos/GizmoMath.h"
+#include "BaseGizmos/GizmoViewContext.h"
 #include "PrimitiveSceneProxy.h"
 #include "SceneManagement.h" // FMeshElementCollector, FPrimitiveDrawInterface
 
