InteractiveToolsFramework: Move GizmoRenderingUtil functions to be nested under UE namespaces and fix usage. Also revert one of the deprecated signatures to take UGizmoViewContext instead of the new interface, to avoid a potential compile error if the function is called somewhere that the class is not fully defined.

diff --git a/Engine/Plugins/Experimental/GizmoEdMode/Source/LightGizmos/Private/DirectionalLightGizmo.cpp b/Engine/Plugins/Experimental/GizmoEdMode/Source/LightGizmos/Private/DirectionalLightGizmo.cpp
--- a/Engine/Plugins/Experimental/GizmoEdMode/Source/LightGizmos/Private/DirectionalLightGizmo.cpp
+++ b/Engine/Plugins/Experimental/GizmoEdMode/Source/LightGizmos/Private/DirectionalLightGizmo.cpp
@@ -111,7 +111,7 @@ void UDirectionalLightGizmo::Render(IToolsContextRenderAPI* RenderAPI)
 	FVector LineEnd = Start + LineDir * ArrowLength;
 
 	const FSceneView* View = RenderAPI->GetSceneView();
-	float PixelToWorld = GizmoRenderingUtil::CalculateLocalPixelToWorldScale(View, LineEnd);
+	float PixelToWorld = UE::GizmoRenderingUtil::CalculateLocalPixelToWorldScale(View, LineEnd);
 
 	// Calculate the "true" end point in world space (We want the line to be the same length in screen space)
 	FVector LineEndWorld = Start + LineDir * ArrowLength * PixelToWorld;
diff --git a/Engine/Plugins/Experimental/GizmoEdMode/Source/LightGizmos/Private/ScalableConeGizmo.cpp b/Engine/Plugins/Experimental/GizmoEdMode/Source/LightGizmos/Private/ScalableConeGizmo.cpp
--- a/Engine/Plugins/Experimental/GizmoEdMode/Source/LightGizmos/Private/ScalableConeGizmo.cpp
+++ b/Engine/Plugins/Experimental/GizmoEdMode/Source/LightGizmos/Private/ScalableConeGizmo.cpp
@@ -81,7 +81,7 @@ void UScalableConeGizmo::Render(IToolsContextRenderAPI* RenderAPI)
 
 			// Get the Pixel to World scale of the line
 			const FSceneView* View = RenderAPI->GetSceneView();
-			float PixelToWorld = GizmoRenderingUtil::CalculateLocalPixelToWorldScale(View, LineEnd);
+			float PixelToWorld = UE::GizmoRenderingUtil::CalculateLocalPixelToWorldScale(View, LineEnd);
 
 			// Draw the lines in both directions
 			RenderAPI->GetPrimitiveDrawInterface()->DrawLine(LineStart, LineStart + HitAxis * PixelToWorld * LineLength, FLinearColor::Red, SDPG_Foreground);
diff --git a/Engine/Source/Editor/Experimental/EditorInteractiveToolsFramework/Private/EditorGizmos/TransformGizmo.cpp b/Engine/Source/Editor/Experimental/EditorInteractiveToolsFramework/Private/EditorGizmos/TransformGizmo.cpp
--- a/Engine/Source/Editor/Experimental/EditorInteractiveToolsFramework/Private/EditorGizmos/TransformGizmo.cpp
+++ b/Engine/Source/Editor/Experimental/EditorInteractiveToolsFramework/Private/EditorGizmos/TransformGizmo.cpp
@@ -2450,7 +2450,7 @@ void UTransformGizmo::OnClickReleaseArcBallRotate(const FInputDeviceRay& Release
 
 float UTransformGizmo::GetWorldRadius(const float InRadius) const
 {
-	const float PixelToWorldScale = GizmoRenderingUtil::CalculateLocalPixelToWorldScale(GizmoViewContext, CurrentTransform.GetLocation());
+	const float PixelToWorldScale = UE::GizmoRenderingUtil::CalculateLocalPixelToWorldScale(GizmoViewContext, CurrentTransform.GetLocation());
 	const float GizmoScale = TransformGizmoSource ? TransformGizmoSource->GetGizmoScale() : 1.0f;
 	return InRadius * GetSizeCoefficient() * PixelToWorldScale * GizmoScale;
 }
diff --git a/Engine/Source/Runtime/InteractiveToolsFramework/Private/BaseGizmos/GizmoArrowComponent.cpp b/Engine/Source/Runtime/InteractiveToolsFramework/Private/BaseGizmos/GizmoArrowComponent.cpp
--- a/Engine/Source/Runtime/InteractiveToolsFramework/Private/BaseGizmos/GizmoArrowComponent.cpp
+++ b/Engine/Source/Runtime/InteractiveToolsFramework/Private/BaseGizmos/GizmoArrowComponent.cpp
@@ -43,7 +43,7 @@ namespace GizmoArrowComponentLocals
 				return false;
 			}
 
-			LengthScale = GizmoRenderingUtil::CalculateLocalPixelToWorldScale(View, WorldOrigin);
+			LengthScale = UE::GizmoRenderingUtil::CalculateLocalPixelToWorldScale(View, WorldOrigin);
 		}
 
 		PixelToWorldScaleOut = LengthScale;
diff --git a/Engine/Source/Runtime/InteractiveToolsFramework/Private/BaseGizmos/GizmoBoxComponent.cpp b/Engine/Source/Runtime/InteractiveToolsFramework/Private/BaseGizmos/GizmoBoxComponent.cpp
--- a/Engine/Source/Runtime/InteractiveToolsFramework/Private/BaseGizmos/GizmoBoxComponent.cpp
+++ b/Engine/Source/Runtime/InteractiveToolsFramework/Private/BaseGizmos/GizmoBoxComponent.cpp
@@ -117,7 +117,7 @@ public:
 				float LengthScale = 1;
 				if (bIsViewDependent)
 				{
-					LengthScale = GizmoRenderingUtil::CalculateLocalPixelToWorldScale(View, WorldOrigin);
+					LengthScale = UE::GizmoRenderingUtil::CalculateLocalPixelToWorldScale(View, WorldOrigin);
 				}
 
 				FVector WorldCenter = WorldOrigin
@@ -310,7 +310,7 @@ bool UGizmoBoxComponent::LineTraceComponent(FHitResult& OutHit, const FVector St
 
 	if (bIsViewDependent)
 	{
-		DynamicPixelToWorldScale = GizmoRenderingUtil::CalculateLocalPixelToWorldScale(
+		DynamicPixelToWorldScale = UE::GizmoRenderingUtil::CalculateLocalPixelToWorldScale(
 			GizmoViewContext, WorldOrigin);
 	}
 
diff --git a/Engine/Source/Runtime/InteractiveToolsFramework/Private/BaseGizmos/GizmoCircleComponent.cpp b/Engine/Source/Runtime/InteractiveToolsFramework/Private/BaseGizmos/GizmoCircleComponent.cpp
--- a/Engine/Source/Runtime/InteractiveToolsFramework/Private/BaseGizmos/GizmoCircleComponent.cpp
+++ b/Engine/Source/Runtime/InteractiveToolsFramework/Private/BaseGizmos/GizmoCircleComponent.cpp
@@ -71,7 +71,7 @@ public:
 				bool bIsViewDependent = (bExternalIsViewDependent) ? (*bExternalIsViewDependent) : false;
 				if (bIsViewDependent)
 				{
-					LengthScale = GizmoRenderingUtil::CalculateLocalPixelToWorldScale(View, Origin);
+					LengthScale = UE::GizmoRenderingUtil::CalculateLocalPixelToWorldScale(View, Origin);
 				}
 
 				double UseRadius = LengthScale * Radius;
@@ -268,7 +268,7 @@ bool UGizmoCircleComponent::LineTraceComponent(FHitResult& OutHit, const FVector
 			return false;
 		}
 
-		PixelToWorldScale = GizmoRenderingUtil::CalculateLocalPixelToWorldScale(GizmoViewContext, WorldOrigin);
+		PixelToWorldScale = UE::GizmoRenderingUtil::CalculateLocalPixelToWorldScale(GizmoViewContext, WorldOrigin);
 	}
 
 	float LengthScale = PixelToWorldScale;
diff --git a/Engine/Source/Runtime/InteractiveToolsFramework/Private/BaseGizmos/GizmoLineHandleComponent.cpp b/Engine/Source/Runtime/InteractiveToolsFramework/Private/BaseGizmos/GizmoLineHandleComponent.cpp
--- a/Engine/Source/Runtime/InteractiveToolsFramework/Private/BaseGizmos/GizmoLineHandleComponent.cpp
+++ b/Engine/Source/Runtime/InteractiveToolsFramework/Private/BaseGizmos/GizmoLineHandleComponent.cpp
@@ -63,7 +63,7 @@ public:
 				bool bIsViewDependent = (bExternalIsViewDependent) ? (*bExternalIsViewDependent) : false;
 				if (bIsViewDependent)
 				{
-					LengthScale = GizmoRenderingUtil::CalculateLocalPixelToWorldScale(View, WorldDiskOrigin);
+					LengthScale = UE::GizmoRenderingUtil::CalculateLocalPixelToWorldScale(View, WorldDiskOrigin);
 				}
 
 				FVector ScaledIntevalStart = -LengthScale * (WorldIntervalEnd - WorldDiskOrigin) + WorldDiskOrigin;
@@ -170,7 +170,7 @@ bool UGizmoLineHandleComponent::LineTraceComponent(FHitResult& OutHit, const FVe
 	float PixelToWorldScale = 1;
 	if (bIsViewDependent)
 	{
-		PixelToWorldScale = GizmoRenderingUtil::CalculateLocalPixelToWorldScale(GizmoViewContext, WorldBaseOrigin);
+		PixelToWorldScale = UE::GizmoRenderingUtil::CalculateLocalPixelToWorldScale(GizmoViewContext, WorldBaseOrigin);
 	}
 
 	float LengthScale = (bImageScale) ? PixelToWorldScale : 1.f;
diff --git a/Engine/Source/Runtime/InteractiveToolsFramework/Private/BaseGizmos/GizmoRectangleComponent.cpp b/Engine/Source/Runtime/InteractiveToolsFramework/Private/BaseGizmos/GizmoRectangleComponent.cpp
--- a/Engine/Source/Runtime/InteractiveToolsFramework/Private/BaseGizmos/GizmoRectangleComponent.cpp
+++ b/Engine/Source/Runtime/InteractiveToolsFramework/Private/BaseGizmos/GizmoRectangleComponent.cpp
@@ -59,7 +59,7 @@ namespace GizmoRectangleComponentLocals
 				return false;
 			}
 
-			LengthScale = GizmoRenderingUtil::CalculateLocalPixelToWorldScale(View, WorldOrigin);
+			LengthScale = UE::GizmoRenderingUtil::CalculateLocalPixelToWorldScale(View, WorldOrigin);
 		}
 
 		double UseOffsetX = LengthScale * OffsetX;
diff --git a/Engine/Source/Runtime/InteractiveToolsFramework/Private/BaseGizmos/GizmoRenderingUtil.cpp b/Engine/Source/Runtime/InteractiveToolsFramework/Private/BaseGizmos/GizmoRenderingUtil.cpp
--- a/Engine/Source/Runtime/InteractiveToolsFramework/Private/BaseGizmos/GizmoRenderingUtil.cpp
+++ b/Engine/Source/Runtime/InteractiveToolsFramework/Private/BaseGizmos/GizmoRenderingUtil.cpp
@@ -76,13 +76,13 @@ UViewAdjustedStaticMeshGizmoComponent* UE::GizmoRenderingUtil::CreateDefaultMate
 	Component->SetStaticMesh(Mesh);
 
 	Component->SetGizmoViewContext(GizmoViewContext);
-	Component->TranslucencySortPriority = GizmoRenderingUtil::GIZMO_TRANSLUCENCY_SORT_PRIORITY;
+	Component->TranslucencySortPriority = UE::GizmoRenderingUtil::GIZMO_TRANSLUCENCY_SORT_PRIORITY;
 	// Used by the default material to be able to be occluded by other gizmo elements
 	Component->bRenderCustomDepth = true;
 	// Not sure that this actually gets respected in any way for non-PDI calls, but just in case
 	Component->DepthPriorityGroup = SDPG_Foreground;
 
-	Component->SetAllMaterials(GizmoRenderingUtil::GetDefaultGizmoComponentMaterial(Color, Component));
+	Component->SetAllMaterials(UE::GizmoRenderingUtil::GetDefaultGizmoComponentMaterial(Color, Component));
 	if (bAddHoverMaterial)
 	{
 		Component->SetHoverOverrideMaterial(UE::GizmoRenderingUtil::GetDefaultGizmoComponentMaterial(HoverColor, Component));
@@ -100,7 +100,7 @@ UViewAdjustedStaticMeshGizmoComponent* UE::GizmoRenderingUtil::CreateDefaultMate
 }
 
 
-float GizmoRenderingUtil::CalculateLocalPixelToWorldScale(
+float UE::GizmoRenderingUtil::CalculateLocalPixelToWorldScale(
 	const FSceneView* View,
 	const FVector& Location)
 {
@@ -115,7 +115,7 @@ float GizmoRenderingUtil::CalculateLocalPixelToWorldScale(
 	return CalculateLocalPixelToWorldScale(&Wrapper, Location);
 }
 
-float GizmoRenderingUtil::CalculateLocalPixelToWorldScale(
+float UE::GizmoRenderingUtil::CalculateLocalPixelToWorldScale(
 	const UE::GizmoRenderingUtil::ISceneViewInterface* View, 
 	const FVector& Location)
 {
@@ -139,7 +139,7 @@ float GizmoRenderingUtil::CalculateLocalPixelToWorldScale(
 	return (float)(sqrt(WorldDeltaSqr / PixelDeltaSqr));
 }
 
-float GizmoRenderingUtil::CalculateViewDependentScaleAndFlatten(
+float UE::GizmoRenderingUtil::CalculateViewDependentScaleAndFlatten(
 	const FSceneView* View,
 	const FVector& Location,
 	const float InScale,
@@ -212,4 +212,26 @@ FLinearColor UE::GizmoRenderingUtil::GetDefaultAxisColor(EAxis::Type Axis)
 	default:
 		return FLinearColor::White;
 	}
+}
+
+// Forward the deprecated methods (that are in the wrong namespace) to the proper namespace methods
+float GizmoRenderingUtil::CalculateLocalPixelToWorldScale(
+	const FSceneView* View,
+	const FVector& Location)
+{
+	return UE::GizmoRenderingUtil::CalculateLocalPixelToWorldScale(View, Location);
+}
+float GizmoRenderingUtil::CalculateLocalPixelToWorldScale(
+	const UGizmoViewContext* ViewContext,
+	const FVector& Location)
+{
+	return UE::GizmoRenderingUtil::CalculateLocalPixelToWorldScale(ViewContext, Location);
+}
+float GizmoRenderingUtil::CalculateViewDependentScaleAndFlatten(
+	const FSceneView* View,
+	const FVector& Location,
+	const float InScale,
+	FVector& OutFlattenScale)
+{
+	return UE::GizmoRenderingUtil::CalculateViewDependentScaleAndFlatten(View, Location, InScale, OutFlattenScale);
 }
\ No newline at end of file
diff --git a/Engine/Source/Runtime/InteractiveToolsFramework/Private/BaseGizmos/ScalableSphereGizmo.cpp b/Engine/Source/Runtime/InteractiveToolsFramework/Private/BaseGizmos/ScalableSphereGizmo.cpp
--- a/Engine/Source/Runtime/InteractiveToolsFramework/Private/BaseGizmos/ScalableSphereGizmo.cpp
+++ b/Engine/Source/Runtime/InteractiveToolsFramework/Private/BaseGizmos/ScalableSphereGizmo.cpp
@@ -58,7 +58,7 @@ void UScalableSphereGizmo::Render(IToolsContextRenderAPI* RenderAPI)
 
 			// Get the Pixel to World scale of the line
 			const FSceneView* View = RenderAPI->GetSceneView();
-			float PixelToWorld = GizmoRenderingUtil::CalculateLocalPixelToWorldScale(View, LineEnd);
+			float PixelToWorld = UE::GizmoRenderingUtil::CalculateLocalPixelToWorldScale(View, LineEnd);
 
 			// Draw the lines in both directions
 			RenderAPI->GetPrimitiveDrawInterface()->DrawLine(LineStart, LineStart + ActiveAxis * PixelToWorld * LineLength, FLinearColor::Red, SDPG_Foreground);
diff --git a/Engine/Source/Runtime/InteractiveToolsFramework/Private/BaseGizmos/ViewBasedTransformAdjusters.cpp b/Engine/Source/Runtime/InteractiveToolsFramework/Private/BaseGizmos/ViewBasedTransformAdjusters.cpp
--- a/Engine/Source/Runtime/InteractiveToolsFramework/Private/BaseGizmos/ViewBasedTransformAdjusters.cpp
+++ b/Engine/Source/Runtime/InteractiveToolsFramework/Private/BaseGizmos/ViewBasedTransformAdjusters.cpp
@@ -28,7 +28,7 @@ namespace ViewBasedTransformAdjustersLocals
 		//  to it. However the approach here tends to be the least broken overall.
 		double DistanceToGizmo = FVector3d::Distance(GizmoToWorld.GetLocation(), View.GetViewLocation());
 		FVector3d PointAtWhichToCheckScale = View.GetViewLocation() + DistanceToGizmo * View.GetViewDirection();
-		return GizmoRenderingUtil::CalculateLocalPixelToWorldScale(&View, PointAtWhichToCheckScale);
+		return UE::GizmoRenderingUtil::CalculateLocalPixelToWorldScale(&View, PointAtWhichToCheckScale);
 	}
 
 	// Apply the settings to get a transform modified by view
@@ -90,7 +90,7 @@ namespace ViewBasedTransformAdjustersLocals
 FTransform UE::GizmoRenderingUtil::FSimpleConstantViewScaleAdjuster::GetAdjustedComponentToWorld(
 	const ISceneViewInterface& View, const FTransform& CurrentComponentToWorld)
 {
-	double ExtraScaling = ::GizmoRenderingUtil::CalculateLocalPixelToWorldScale(&View, CurrentComponentToWorld.GetLocation());
+	double ExtraScaling = UE::GizmoRenderingUtil::CalculateLocalPixelToWorldScale(&View, CurrentComponentToWorld.GetLocation());
 	FTransform AdjustedTransform = CurrentComponentToWorld;
 	AdjustedTransform.MultiplyScale3D(FVector(ExtraScaling));
 
diff --git a/Engine/Source/Runtime/InteractiveToolsFramework/Public/BaseGizmos/GizmoElementBase.h b/Engine/Source/Runtime/InteractiveToolsFramework/Public/BaseGizmos/GizmoElementBase.h
--- a/Engine/Source/Runtime/InteractiveToolsFramework/Public/BaseGizmos/GizmoElementBase.h
+++ b/Engine/Source/Runtime/InteractiveToolsFramework/Public/BaseGizmos/GizmoElementBase.h
@@ -67,7 +67,7 @@ public:
 		void Initialize(const FSceneView* InSceneView, FTransform InTransform)
 		{
 			LocalToWorldTransform = InTransform;
-			PixelToWorldScale = GizmoRenderingUtil::CalculateLocalPixelToWorldScale(InSceneView, InTransform.GetLocation());
+			PixelToWorldScale = UE::GizmoRenderingUtil::CalculateLocalPixelToWorldScale(InSceneView, InTransform.GetLocation());
 		}
 
 		// Returns the mesh material based on the current interaction state.
@@ -106,7 +106,7 @@ public:
 		{
 			check(InGizmoViewContext);
 			LocalToWorldTransform = InTransform;
-			PixelToWorldScale = GizmoRenderingUtil::CalculateLocalPixelToWorldScale(InGizmoViewContext, InTransform.GetLocation());
+			PixelToWorldScale = UE::GizmoRenderingUtil::CalculateLocalPixelToWorldScale(InGizmoViewContext, InTransform.GetLocation());
 			bIsPerspectiveProjection = InGizmoViewContext->IsPerspectiveProjection();
 		}
 	};
diff --git a/Engine/Source/Runtime/InteractiveToolsFramework/Public/BaseGizmos/GizmoRenderingUtil.h b/Engine/Source/Runtime/InteractiveToolsFramework/Public/BaseGizmos/GizmoRenderingUtil.h
--- a/Engine/Source/Runtime/InteractiveToolsFramework/Public/BaseGizmos/GizmoRenderingUtil.h
+++ b/Engine/Source/Runtime/InteractiveToolsFramework/Public/BaseGizmos/GizmoRenderingUtil.h
@@ -91,11 +91,7 @@ namespace UE::GizmoRenderingUtil
 	 * Gets a red/green/blue color based on the axis (X, Y, or Z).
 	 */
 	INTERACTIVETOOLSFRAMEWORK_API FLinearColor GetDefaultAxisColor(EAxis::Type Axis);
-}
 
-//~ TODO: These should someday be moved to nest under the UE namespace above
-namespace GizmoRenderingUtil
-{
 	/**
 	 * @return Conversion factor between pixel and world-space coordinates at 3D point Location in View.
 	 * @warning This is a local estimate and is increasingly incorrect as the 3D point gets further from Location
@@ -118,3 +114,24 @@ namespace GizmoRenderingUtil
 		const float Scale,
 		FVector& OutWorldFlattenScale);
 }
+
+// This namespace is deprecated- use UE::GizmoRenderingUtil instead.
+namespace GizmoRenderingUtil
+{
+	UE_DEPRECATED(5.5, "This function was moved to the UE::GizmoRenderingUtil namespace.")
+	INTERACTIVETOOLSFRAMEWORK_API float CalculateLocalPixelToWorldScale(
+		const FSceneView* View,
+		const FVector& Location);
+	
+	UE_DEPRECATED(5.5, "This function was moved to the UE::GizmoRenderingUtil namespace.")
+	INTERACTIVETOOLSFRAMEWORK_API float CalculateLocalPixelToWorldScale(
+		const UGizmoViewContext* ViewContext,
+		const FVector& Location);
+
+	UE_DEPRECATED(5.5, "This function was moved to the UE::GizmoRenderingUtil namespace.")
+	INTERACTIVETOOLSFRAMEWORK_API float CalculateViewDependentScaleAndFlatten(
+		const FSceneView* View,
+		const FVector& Location,
+		const float Scale,
+		FVector& OutWorldFlattenScale);
+}
diff --git a/Engine/Source/Runtime/InteractiveToolsFramework/Public/BaseGizmos/ViewBasedTransformAdjusters.h b/Engine/Source/Runtime/InteractiveToolsFramework/Public/BaseGizmos/ViewBasedTransformAdjusters.h
--- a/Engine/Source/Runtime/InteractiveToolsFramework/Public/BaseGizmos/ViewBasedTransformAdjusters.h
+++ b/Engine/Source/Runtime/InteractiveToolsFramework/Public/BaseGizmos/ViewBasedTransformAdjusters.h
@@ -119,7 +119,7 @@ namespace UE::GizmoRenderingUtil
 
 		// IGizmoViewTransformAdjuster
 		INTERACTIVETOOLSFRAMEWORK_API virtual FTransform GetAdjustedComponentToWorld(
-			const GizmoRenderingUtil::ISceneViewInterface& View,
+			const UE::GizmoRenderingUtil::ISceneViewInterface& View,
 			const FTransform& CurrentComponentToWorld) override;
 
 		INTERACTIVETOOLSFRAMEWORK_API virtual FTransform GetAdjustedComponentToWorld_RenderThread(
@@ -158,7 +158,7 @@ namespace UE::GizmoRenderingUtil
 
 		// IGizmoViewTransformAdjuster
 		INTERACTIVETOOLSFRAMEWORK_API virtual FTransform GetAdjustedComponentToWorld(
-			const GizmoRenderingUtil::ISceneViewInterface& View,
+			const UE::GizmoRenderingUtil::ISceneViewInterface& View,
 			const FTransform& CurrentComponentToWorld) override;
 	private:
 		FTransform ViewRelativeTransform = FTransform::Identity;
