InteractiveToolsFramework: Fix gizmo scaling not being applied correctly after last refactor.

diff --git a/Engine/Source/Runtime/InteractiveToolsFramework/Private/BaseGizmos/AxisPositionGizmo.cpp b/Engine/Source/Runtime/InteractiveToolsFramework/Private/BaseGizmos/AxisPositionGizmo.cpp
--- a/Engine/Source/Runtime/InteractiveToolsFramework/Private/BaseGizmos/AxisPositionGizmo.cpp
+++ b/Engine/Source/Runtime/InteractiveToolsFramework/Private/BaseGizmos/AxisPositionGizmo.cpp
@@ -108,8 +108,8 @@ bool UAxisPositionGizmo::InitializeAsScaleGizmo(
 		// Create new and add to shared state.
 		USceneComponent* RootComponent = Params.Component->GetOwner()->GetRootComponent();
 		UGizmoComponentAxisSource* CastAxisSource = UGizmoComponentAxisSource::Construct(RootComponent, AxisIndex,
-			// bUseLocalAxes, not important because we're going to be updating this value every tick
-			true,
+			// bUseLocalAxes, has to be false to always use cardinal axes
+			false,
 			Owner);
 		UnitCardinalAxisSource = CastAxisSource;
 		if (SharedState)
