InteractiveToolsFramework: Adds support for scale snapping

diff --git a/Engine/Source/Editor/Experimental/EditorInteractiveToolsFramework/Private/EdModeInteractiveToolsContext.cpp b/Engine/Source/Editor/Experimental/EditorInteractiveToolsFramework/Private/EdModeInteractiveToolsContext.cpp
--- a/Engine/Source/Editor/Experimental/EditorInteractiveToolsFramework/Private/EdModeInteractiveToolsContext.cpp
+++ b/Engine/Source/Editor/Experimental/EditorInteractiveToolsFramework/Private/EdModeInteractiveToolsContext.cpp
@@ -166,6 +166,8 @@ public:
 		Config.PositionGridDimensions = FVector(EditorGridSize, EditorGridSize, EditorGridSize);
 		Config.bEnableRotationGridSnapping = (GetDefault<ULevelEditorViewportSettings>()->RotGridEnabled != 0);
 		Config.RotationGridAngles = GEditor->GetRotGridSize();
+		Config.bEnableScaleGridSnapping = (GetDefault<ULevelEditorViewportSettings>()->SnapScaleEnabled != 0);
+		Config.ScaleGridSize = GEditor->GetScaleGridSize();
 		Config.bEnableAbsoluteWorldSnapping = ToolsContext->GetAbsoluteWorldSnappingEnabled();
 		return Config;
 	}
diff --git a/Engine/Source/Runtime/InteractiveToolsFramework/Private/BaseGizmos/CombinedTransformGizmo.cpp b/Engine/Source/Runtime/InteractiveToolsFramework/Private/BaseGizmos/CombinedTransformGizmo.cpp
--- a/Engine/Source/Runtime/InteractiveToolsFramework/Private/BaseGizmos/CombinedTransformGizmo.cpp
+++ b/Engine/Source/Runtime/InteractiveToolsFramework/Private/BaseGizmos/CombinedTransformGizmo.cpp
@@ -1142,6 +1142,15 @@ UInteractiveGizmo* UCombinedTransformGizmo::AddAxisScaleGizmo(
 	}
 	ensure(Gizmo->InitializeAsScaleGizmo(Params, bDisallowNegativeScaling, &SharedState));
 
+	if (UGizmoAxisScaleParameterSource* ParameterSource = Cast<UGizmoAxisScaleParameterSource>(Gizmo->ParameterSource.GetObject()))
+	{
+		ParameterSource->ScaleAxisDeltaConstraintFunction = [this](const double ScaleAxisDelta, double& SnappedScaleAxisDelta) { return ScaleAxisDeltaSnapFunction(ScaleAxisDelta, SnappedScaleAxisDelta); };
+	}
+	else
+	{
+		ensure(false);
+	}
+	
 	NonUniformScaleSubGizmos.Add(FSubGizmoInfo{ Params.Component, Gizmo });
 	ActiveComponents.Add(Params.Component);
 	ActiveGizmos.Add(Gizmo);
@@ -1159,6 +1168,16 @@ UInteractiveGizmo* UCombinedTransformGizmo::AddPlaneScaleGizmo(
 	}
 	ensure(Gizmo->InitializeAsScaleGizmo(Params, bDisallowNegativeScaling, &SharedState));
 
+	if (UGizmoPlaneScaleParameterSource* ParameterSource = Cast<UGizmoPlaneScaleParameterSource>(Gizmo->ParameterSource.GetObject()))
+	{
+		ParameterSource->ScaleAxisXDeltaConstraintFunction = [this](double ScaleAxisDelta, double& SnappedScaleAxisDelta) { return ScaleAxisDeltaSnapFunction(ScaleAxisDelta, SnappedScaleAxisDelta); };
+		ParameterSource->ScaleAxisYDeltaConstraintFunction = [this](double ScaleAxisDelta, double& SnappedScaleAxisDelta) { return ScaleAxisDeltaSnapFunction(ScaleAxisDelta, SnappedScaleAxisDelta); };
+	}
+	else
+	{
+		ensure(false);
+	}
+
 	NonUniformScaleSubGizmos.Add(FSubGizmoInfo{ Params.Component, Gizmo });
 	ActiveComponents.Add(Params.Component);
 	ActiveGizmos.Add(Gizmo);
@@ -1176,6 +1195,15 @@ UInteractiveGizmo* UCombinedTransformGizmo::AddUniformScaleGizmo(
 	}
 	ensure(Gizmo->InitializeAsUniformScaleGizmo(Params, bDisallowNegativeScaling, &SharedState));
 
+	if (UGizmoUniformScaleParameterSource* ParameterSource = Cast<UGizmoUniformScaleParameterSource>(Gizmo->ParameterSource.GetObject()))
+	{
+		ParameterSource->ScaleAxisDeltaConstraintFunction = [this](const double ScaleAxisDelta, double& SnappedScaleAxisDelta) { return ScaleAxisDeltaSnapFunction(ScaleAxisDelta, SnappedScaleAxisDelta); };
+	}
+	else
+	{
+		ensure(false);
+	}
+	
 	UniformScaleSubGizmos.Add(FSubGizmoInfo{ Params.Component, Gizmo });
 	ActiveComponents.Add(Params.Component);
 	ActiveGizmos.Add(Gizmo);
@@ -1322,6 +1350,7 @@ UInteractiveGizmo* UCombinedTransformGizmo::AddAxisScaleGizmo(
 
 	// parameter source maps axis-parameter-change to translation of TransformSource's transform
 	UGizmoAxisScaleParameterSource* ParamSource = UGizmoAxisScaleParameterSource::Construct(ParameterAxisSource, TransformSource, this);
+	ParamSource->ScaleAxisDeltaConstraintFunction = [this](const double ScaleAxisDelta, double& SnappedScaleAxisDelta) { return ScaleAxisDeltaSnapFunction(ScaleAxisDelta, SnappedScaleAxisDelta); };
 	ParamSource->bClampToZero = bDisallowNegativeScaling;
 	ScaleGizmo->ParameterSource = ParamSource;
 
@@ -1354,6 +1383,8 @@ UInteractiveGizmo* UCombinedTransformGizmo::AddPlaneScaleGizmo(
 
 	// parameter source maps axis-parameter-change to translation of TransformSource's transform
 	UGizmoPlaneScaleParameterSource* ParamSource = UGizmoPlaneScaleParameterSource::Construct(ParameterAxisSource, TransformSource, this);
+	ParamSource->ScaleAxisXDeltaConstraintFunction = [this](double ScaleAxisDelta, double& SnappedScaleAxisDelta) { return ScaleAxisDeltaSnapFunction(ScaleAxisDelta, SnappedScaleAxisDelta); };
+	ParamSource->ScaleAxisYDeltaConstraintFunction = [this](double ScaleAxisDelta, double& SnappedScaleAxisDelta) { return ScaleAxisDeltaSnapFunction(ScaleAxisDelta, SnappedScaleAxisDelta); };
 	ParamSource->bClampToZero = bDisallowNegativeScaling;
 	ParamSource->bUseEqualScaling = true;
 	ScaleGizmo->ParameterSource = ParamSource;
@@ -1386,7 +1417,7 @@ UInteractiveGizmo* UCombinedTransformGizmo::AddUniformScaleGizmo(
 
 	// parameter source maps axis-parameter-change to translation of TransformSource's transform
 	UGizmoUniformScaleParameterSource* ParamSource = UGizmoUniformScaleParameterSource::Construct(ParameterAxisSource, TransformSource, this);
-	//ParamSource->PositionConstraintFunction = [this](const FVector& Pos, FVector& Snapped) { return PositionSnapFunction(Pos, Snapped); };
+	ParamSource->ScaleAxisDeltaConstraintFunction = [this](const double ScaleAxisDelta, double& SnappedScaleAxisDelta) { return ScaleAxisDeltaSnapFunction(ScaleAxisDelta, SnappedScaleAxisDelta); };
 	ScaleGizmo->ParameterSource = ParamSource;
 
 	// sub-component provides hit target
@@ -1585,5 +1616,20 @@ bool UCombinedTransformGizmo::RotationAxisAngleSnapFunction(double AxisAngleDelt
 	return false;
 }
 
+bool UCombinedTransformGizmo::ScaleAxisDeltaSnapFunction(const double ScaleAxisDelta, double & SnappedAxisScaleDeltaOut) const
+{
+	if (!bSnapToScaleGrid) return false;
+
+	const FToolContextSnappingConfiguration SnappingConfig = GetGizmoManager()->GetContextQueriesAPI()->GetCurrentSnappingSettings();
+	if ( SnappingConfig.bEnableScaleGridSnapping )
+	{
+		const double SnapDelta = SnappingConfig.ScaleGridSize;
+		SnappedAxisScaleDeltaOut = UE::Geometry::SnapToIncrement(ScaleAxisDelta, SnapDelta);
+		return true;
+	}
+	return false;
+}
+
+
 #undef LOCTEXT_NAMESPACE
 
diff --git a/Engine/Source/Runtime/InteractiveToolsFramework/Private/BaseGizmos/ParameterToTransformAdapters.cpp b/Engine/Source/Runtime/InteractiveToolsFramework/Private/BaseGizmos/ParameterToTransformAdapters.cpp
--- a/Engine/Source/Runtime/InteractiveToolsFramework/Private/BaseGizmos/ParameterToTransformAdapters.cpp
+++ b/Engine/Source/Runtime/InteractiveToolsFramework/Private/BaseGizmos/ParameterToTransformAdapters.cpp
@@ -183,17 +183,25 @@ void UGizmoUniformScaleParameterSource::SetParameter(const FVector2D& NewValue)
 	SignedDelta += 1.0;
 
 	FTransform NewTransform = InitialTransform;
-	FVector StartScale = InitialTransform.GetScale3D();
-	FVector NewScale = SignedDelta * StartScale;
-	NewTransform.SetScale3D(NewScale);
+	const FVector StartScale = InitialTransform.GetScale3D();
+	
+	double SnappedDelta;
+	
+	// if using snapping while scaling
+	if (ScaleAxisDeltaConstraintFunction(SignedDelta, SnappedDelta))
+	{
+		SignedDelta = SnappedDelta;
+	}
+	// ensures that all 3 axes scale proportionally while following closest snap factor
+	// ex: Scale Snap is set to 1, StartScale = (1,2,3), when uniform scaling NewScale=(2, 4, 6), instead of (2, 3, 4) to retain proportions
+	const FVector NewScale = SignedDelta * StartScale;
 
-	// apply position constraint
-	//FVector SnappedPos;
-	//if (PositionConstraintFunction(NewTransform.GetScale(), SnappedPos))
-	//{
-	//	FVector PlanePos = GizmoMath::ProjectPointOntoPlane(SnappedPos, CurTranslationOrigin, CurTranslationNormal);
-	//	NewTransform.SetTranslation(PlanePos);
-	//}
+	// currently calling ScaleConstraintFunction has no effect (no changes made to SignedDelta) because this constraint function
+	// is intended to relate to WorldGridSnapping. Currently WorldGridSnapping has no effect on scaling, with or without normal snapping
+	// because the viewport scale mode fixes the transform space to local
+	SignedDelta = ScaleConstraintFunction(SignedDelta);
+
+	NewTransform.SetScale3D(NewScale);
 
 	TransformSource->SetTransform(NewTransform);
 
@@ -225,14 +233,37 @@ void UGizmoAxisScaleParameterSource::SetParameter(float NewValue)
 {
 	Parameter = NewValue;
 	LastChange.CurrentValue = NewValue;
+	
+	double ScaleDelta = LastChange.GetChangeDelta();
+	ScaleDelta *= ScaleMultiplier;
+	ScaleDelta += 1.0;
 
-	// construct translation as delta from initial position
-	float ScaleDelta = LastChange.GetChangeDelta() * ScaleMultiplier;
-
-	// translate the initial transform
 	FTransform NewTransform = InitialTransform;
-	FVector StartScale = InitialTransform.GetScale3D();
-	FVector NewScale = StartScale + ScaleDelta * CurScaleAxis;
+	const FVector StartScale = InitialTransform.GetScale3D();
+
+	FVector NewScale;
+	double SnappedDelta;
+	
+	// check for any constraints on the delta value
+	if (ScaleAxisDeltaConstraintFunction(ScaleDelta, SnappedDelta))
+	{
+		ScaleDelta = SnappedDelta;
+		const FVector ScaleDeltaVector = ScaleDelta * CurScaleAxis;
+		// uses addition when snapping
+		// ex: Scale Snap is set to 1, StartScale=(2,2,2), when scaling X axis NewScale=(3,2,2), will NOT be (4,2,2)
+		// Note: Plane/Uniform Scale Snapping not implemented in this way because they
+		// need to preserve proportional relationship between 2+ axes, therefore use multiplication
+		NewScale = StartScale + (ScaleDeltaVector + CurScaleAxis);
+	}
+	else
+	{
+		NewScale = StartScale + ScaleDelta * CurScaleAxis;
+	}
+
+	// currently calling ScaleConstraintFunction has no effect (no changes made to ScaleDelta) because this constraint function
+	// is intended to relate to WorldGridSnapping. Currently WorldGridSnapping has no effect on scaling, with or without normal snapping
+	// because the viewport scale mode fixes the transform space to local
+	ScaleDelta = ScaleConstraintFunction(ScaleDelta);
 
 	if (bClampToZero)
 	{
@@ -241,15 +272,6 @@ void UGizmoAxisScaleParameterSource::SetParameter(float NewValue)
 
 	NewTransform.SetScale3D(NewScale);
 
-
-	// apply position constraint
-	//FVector SnappedPos;
-	//if (ScaleConstraintFunction(NewTransform.GetTranslation(), SnappedPos))
-	//{
-	//	FVector SnappedLinePos = GizmoMath::ProjectPointOntoLine(SnappedPos, CurTranslationOrigin, CurTranslationAxis);
-	//	NewTransform.SetTranslation(SnappedLinePos);
-	//}
-
 	TransformSource->SetTransform(NewTransform);
 
 	OnParameterChanged.Broadcast(this, LastChange);
@@ -281,16 +303,47 @@ void UGizmoPlaneScaleParameterSource::SetParameter(const FVector2D& NewValue)
 	LastChange.CurrentValue = NewValue;
 
 	// construct Scale as delta from initial position
-	FVector2D Delta = LastChange.GetChangeDelta() * ScaleMultiplier;
+	FVector2D ScaleDelta = LastChange.GetChangeDelta() * ScaleMultiplier;
 
 	if (bUseEqualScaling)
 	{
-		Delta = FVector2D(Delta.X + Delta.Y);
+		ScaleDelta = FVector2D(ScaleDelta.X + ScaleDelta.Y);
 	}
 
 	FTransform NewTransform = InitialTransform;
-	FVector StartScale = InitialTransform.GetScale3D();
-	FVector NewScale = StartScale + Delta.X*CurScaleAxisX + Delta.Y*CurScaleAxisY;
+	const FVector StartScale = InitialTransform.GetScale3D();
+
+	double UseScaleDeltaX = ScaleDelta.X;
+	double UseScaleDeltaY = ScaleDelta.Y;
+	
+	FVector NewScale;
+
+	if (bUseEqualScaling)
+	{
+		double SnappedDeltaX = 0.0, SnappedDeltaY = 0.0;
+		
+		// if using snapping while scaling on X and Y axis
+		if (ScaleAxisXDeltaConstraintFunction(UseScaleDeltaX, SnappedDeltaX))
+		{
+			UseScaleDeltaX = SnappedDeltaX;
+		}
+		if (ScaleAxisYDeltaConstraintFunction(UseScaleDeltaY, SnappedDeltaY))
+		{
+			UseScaleDeltaY = SnappedDeltaY;
+		}
+		// ensures that 2 axes on plane scale proportionally while following closest snap factor
+		// ex: Scale Snap is set to 1, StartScale = (1,2,3), scaling on X axis NewScale=(1, 4, 6), instead of (1, 3, 4)
+		NewScale = StartScale + (StartScale*(UseScaleDeltaX*CurScaleAxisX)) + (StartScale*(UseScaleDeltaY*CurScaleAxisY));
+	}
+	else
+	{
+		NewScale = StartScale + ScaleDelta.X*CurScaleAxisX + ScaleDelta.Y*CurScaleAxisY;
+	}
+	
+	// currently calling ScaleConstraintFunction has no effect (no changes made to SignedDelta) because this constraint function
+	// is intended to relate to WorldGridSnapping. Currently WorldGridSnapping has no effect on scaling, with or without normal snapping
+	// because the viewport scale mode fixes the transform space to local
+	ScaleDelta = ScaleConstraintFunction(ScaleDelta);
 
 	if (bClampToZero)
 	{
@@ -299,14 +352,6 @@ void UGizmoPlaneScaleParameterSource::SetParameter(const FVector2D& NewValue)
 
 	NewTransform.SetScale3D(NewScale);
 
-	// apply position constraint
-	//FVector SnappedPos;
-	//if (PositionConstraintFunction(NewTransform.GetScale(), SnappedPos))
-	//{
-	//	FVector PlanePos = GizmoMath::ProjectPointOntoPlane(SnappedPos, CurTranslationOrigin, CurTranslationNormal);
-	//	NewTransform.SetTranslation(PlanePos);
-	//}
-
 	TransformSource->SetTransform(NewTransform);
 
 	OnParameterChanged.Broadcast(this, LastChange);
diff --git a/Engine/Source/Runtime/InteractiveToolsFramework/Public/BaseGizmos/CombinedTransformGizmo.h b/Engine/Source/Runtime/InteractiveToolsFramework/Public/BaseGizmos/CombinedTransformGizmo.h
--- a/Engine/Source/Runtime/InteractiveToolsFramework/Public/BaseGizmos/CombinedTransformGizmo.h
+++ b/Engine/Source/Runtime/InteractiveToolsFramework/Public/BaseGizmos/CombinedTransformGizmo.h
@@ -484,11 +484,16 @@ public:
 	FRotator ExplicitRotationGridSize;
 
 	/**
-	 * If true, then when using world frame, Axis and Plane translation snap to the world grid via the ContextQueriesAPI (in RotationSnapFunction)
+	 * If true, then when using world frame, Axis and Plane rotation snap to the world grid via the ContextQueriesAPI (in RotationSnapFunction)
 	 */
 	UPROPERTY()
 	bool bSnapToWorldRotGrid = true;
 
+	/**
+	 * If true, scaling snaps to the grid
+	 */
+	UPROPERTY()
+	bool bSnapToScaleGrid = true;
 
 
 	/**
@@ -717,6 +722,11 @@ protected:
 	INTERACTIVETOOLSFRAMEWORK_API bool PositionAxisDeltaSnapFunction(double AxisDelta, double& SnappedDeltaOut, int AxisIndex) const;
 	INTERACTIVETOOLSFRAMEWORK_API FQuat RotationSnapFunction(const FQuat& DeltaRotation) const;
 	INTERACTIVETOOLSFRAMEWORK_API bool RotationAxisAngleSnapFunction(double AxisAngleDelta, double& SnappedAxisAngleDeltaOut, int AxisIndex) const;
+	
+	// currently not implemented because WorldGridSnapping currently has no affect on scale/scale-snapping
+	INTERACTIVETOOLSFRAMEWORK_API bool ScaleSnapFunction(double DeltaScale) const;
+	// used to execute Plane, Axis, & Uniform scale snapping
+	INTERACTIVETOOLSFRAMEWORK_API bool ScaleAxisDeltaSnapFunction(double ScaleAxisDelta, double & SnappedAxisScaleDeltaOut) const;
 
 
 	UE_DEPRECATED(5.5, "Use FTransformSubGizmoCommonParams overload instead.")
diff --git a/Engine/Source/Runtime/InteractiveToolsFramework/Public/BaseGizmos/ParameterToTransformAdapters.h b/Engine/Source/Runtime/InteractiveToolsFramework/Public/BaseGizmos/ParameterToTransformAdapters.h
--- a/Engine/Source/Runtime/InteractiveToolsFramework/Public/BaseGizmos/ParameterToTransformAdapters.h
+++ b/Engine/Source/Runtime/InteractiveToolsFramework/Public/BaseGizmos/ParameterToTransformAdapters.h
@@ -323,12 +323,17 @@ class UGizmoUniformScaleParameterSource : public UGizmoBaseVec2ParameterSource
 public:
 
 	/**
-	 * Optional position constraint function. Called during interaction with the new transform origin.
-	 * To snap the transform to a new position, return as second value, and return true from your lambda.
-	 * Note that returned snap point will be projected onto the current translation origin/axis.
-	 * @return true if constraint point was found and should be used, false to ignore
+	 * Optional scale axis delta constraint function. This can be used to snap/constrain the "scale delta" value,
+	 * for example to snap it to fixed increments along the axis. 
+	 * @return true if constrained delta was found and should be used, false to ignore
 	 */
-	//TUniqueFunction<bool(const FVector&, FVector&)> PositionConstraintFunction = [](const FVector&, FVector&) { return false; };
+	TUniqueFunction<bool(double ScaleAxisDelta, double& SnappedScaleAxisDelta)> ScaleAxisDeltaConstraintFunction = [](double, double&) { return false; };
+	/**
+	 * Optional scale constraint function. Called during interaction with the scale delta
+	 * To snap the scale delta, return the snapped value
+	 * @return The snapped value of the scale delta
+	 */
+	TUniqueFunction<double(const double&)> ScaleConstraintFunction = [](const double& DeltaScale) { return DeltaScale; };
 
 
 	virtual FVector2D GetParameter() const override
@@ -419,12 +424,17 @@ class UGizmoAxisScaleParameterSource : public UGizmoBaseFloatParameterSource
 public:
 
 	/**
-	 * Optional position constraint function. Called during interaction with the new transform origin.
-	 * To snap the transform to a new position, return as second value, and return true from your lambda.
-	 * Note that returned snap point will be projected onto the current translation origin/axis.
-	 * @return true if constraint point was found and should be used, false to ignore
+	 * Optional scale axis delta constraint function. This can be used to snap/constrain the "scale delta" value,
+	 * for example to snap it to fixed increments along the axis. 
+	 * @return true if constrained delta was found and should be used, false to ignore
+	 */
+	TUniqueFunction<bool(const double& ScaleAxisDelta, double& SnappedScaleAxisDelta)> ScaleAxisDeltaConstraintFunction = [](const double&, double&) { return false; };
+	/**
+	 * Optional scale constraint function. Called during interaction with the scale delta
+	 * To snap the scale delta, return the snapped value
+	 * @return The snapped value of the scale delta
 	 */
-	//TUniqueFunction<bool(const float&, float&)> ScaleConstraintFunction = [](const float&, float&) { return false; };
+	TUniqueFunction<double(const double&)> ScaleConstraintFunction = [](const double& DeltaScale) { return DeltaScale; };
 
 	virtual float GetParameter() const override
 	{
@@ -519,6 +529,21 @@ public:
 	 */
 	TUniqueFunction<bool(const FVector&, FVector&)> PositionConstraintFunction = [](const FVector&, FVector&) { return false; };
 
+	/**
+	 * Optional scale axis delta constraint function. This can be used to snap/constrain the "scale delta" value,
+	 * for example to snap it to fixed increments along the axis. 
+	 * The Axis[X/Y]DeltaConstraintFunctions are used separately for the X and Y axes of the scale translation
+	 * @return true if constrained delta was found and should be used, false to ignore
+	 */
+	TUniqueFunction<bool(double ScaleAxisDelta, double& SnappedScaleAxisDelta)> ScaleAxisXDeltaConstraintFunction = [](double, double&) { return false; };
+	TUniqueFunction<bool(double ScaleAxisDelta, double& SnappedScaleAxisDelta)> ScaleAxisYDeltaConstraintFunction = [](double, double&) { return false; };
+	/**
+	 * Optional scale constraint function. Called during interaction with the scale delta
+	 * To snap the scale delta, return the snapped value
+	 * @return The snapped value of the scale delta as a FVector2D containing the snapped values for the X and Y axes of the scale translation
+	 */
+	TUniqueFunction<FVector2D(const FVector2D&)> ScaleConstraintFunction = [](const FVector2D& DeltaScale) { return DeltaScale; };
+
 
 	virtual FVector2D GetParameter() const override
 	{
diff --git a/Engine/Source/Runtime/InteractiveToolsFramework/Public/ToolContextInterfaces.h b/Engine/Source/Runtime/InteractiveToolsFramework/Public/ToolContextInterfaces.h
--- a/Engine/Source/Runtime/InteractiveToolsFramework/Public/ToolContextInterfaces.h
+++ b/Engine/Source/Runtime/InteractiveToolsFramework/Public/ToolContextInterfaces.h
@@ -169,6 +169,11 @@ struct FToolContextSnappingConfiguration
 	/** Specify rotation snapping step size */
 	FRotator RotationGridAngles = FRotator::ZeroRotator;
 
+	/** Specify whether scale snapping should be applied */
+	bool bEnableScaleGridSnapping = false;
+	/** Specify scale factor */
+	float ScaleGridSize = 1.0f;
+
 	/** Specify whether Position snapping in World Coordinate System should be Absolute or Relative */
 	bool bEnableAbsoluteWorldSnapping = false;
 };
