ITF: Addresses issues where scaling was doubling initially. Also adds better expected behavior when uniform scaling using an additive function when initial scale is uniform

diff --git a/Engine/Source/Runtime/InteractiveToolsFramework/Private/BaseGizmos/ParameterToTransformAdapters.cpp b/Engine/Source/Runtime/InteractiveToolsFramework/Private/BaseGizmos/ParameterToTransformAdapters.cpp
--- a/Engine/Source/Runtime/InteractiveToolsFramework/Private/BaseGizmos/ParameterToTransformAdapters.cpp
+++ b/Engine/Source/Runtime/InteractiveToolsFramework/Private/BaseGizmos/ParameterToTransformAdapters.cpp
@@ -180,21 +180,35 @@ void UGizmoUniformScaleParameterSource::SetParameter(const FVector2D& NewValue)
 	// This possibly be exposed as a TFunction to allow customization?
 	double SignedDelta = LastChange.GetChangeDelta().X + LastChange.GetChangeDelta().Y;
 	SignedDelta *= ScaleMultiplier;
-	SignedDelta += 1.0;
 
 	FTransform NewTransform = InitialTransform;
 	const FVector StartScale = InitialTransform.GetScale3D();
 	
 	double SnappedDelta;
-	
+	bool bIsSnapped = false;
+
 	// if using snapping while scaling
 	if (ScaleAxisDeltaConstraintFunction(SignedDelta, SnappedDelta))
 	{
 		SignedDelta = SnappedDelta;
+		bIsSnapped = true;
+	}
+
+	FVector NewScale;
+	// if the initial scale is uniform and snapping is on, we can use an additive method to scale up or down
+	if (StartScale.IsUniform() && bIsSnapped)
+	{
+		NewScale = FVector(SignedDelta) + StartScale;
+	}
+	// otherwise, we need to use multiplication to scale
+	// ex: initial scale is (1,2,4) and scale delta is .5 -> next incremented scale should be (1.5, 3, 6)
+	//     to preserve proportions. Addition would result in (1.5, 2.5, 4.5) which does not keep original proportions.
+	//     Additionally, using multiplication when scale is uniform would result in an ex where InitScale=(2,2,2) and
+	//     ScaleDelta = .5 where next scale would be (3,3,3), where the intermediate scale of (2.5,2.5,2.5) is unreachable
+	else
+	{
+		NewScale = SignedDelta * StartScale + StartScale;
 	}
-	// ensures that all 3 axes scale proportionally while following closest snap factor
-	// ex: Scale Snap is set to 1, StartScale = (1,2,3), when uniform scaling NewScale=(2, 4, 6), instead of (2, 3, 4) to retain proportions
-	const FVector NewScale = SignedDelta * StartScale;
 
 	// currently calling ScaleConstraintFunction has no effect (no changes made to SignedDelta) because this constraint function
 	// is intended to relate to WorldGridSnapping. Currently WorldGridSnapping has no effect on scaling, with or without normal snapping
@@ -236,7 +250,6 @@ void UGizmoAxisScaleParameterSource::SetParameter(float NewValue)
 	
 	double ScaleDelta = LastChange.GetChangeDelta();
 	ScaleDelta *= ScaleMultiplier;
-	ScaleDelta += 1.0;
 
 	FTransform NewTransform = InitialTransform;
 	const FVector StartScale = InitialTransform.GetScale3D();
@@ -248,16 +261,13 @@ void UGizmoAxisScaleParameterSource::SetParameter(float NewValue)
 	if (ScaleAxisDeltaConstraintFunction(ScaleDelta, SnappedDelta))
 	{
 		ScaleDelta = SnappedDelta;
-		const FVector ScaleDeltaVector = ScaleDelta * CurScaleAxis;
-		// uses addition when snapping
-		// ex: Scale Snap is set to 1, StartScale=(2,2,2), when scaling X axis NewScale=(3,2,2), will NOT be (4,2,2)
-		// Note: Plane/Uniform Scale Snapping not implemented in this way because they
-		// need to preserve proportional relationship between 2+ axes, therefore use multiplication
-		NewScale = StartScale + (ScaleDeltaVector + CurScaleAxis);
+		// use additive scaling when snap is on
+		NewScale = StartScale + ScaleDelta * CurScaleAxis;
 	}
 	else
 	{
-		NewScale = StartScale + ScaleDelta * CurScaleAxis;
+		// use multiplicative scaling when snap is off
+		NewScale = StartScale * (FVector3d{1} + (ScaleDelta * CurScaleAxis));
 	}
 
 	// currently calling ScaleConstraintFunction has no effect (no changes made to ScaleDelta) because this constraint function
@@ -321,23 +331,45 @@ void UGizmoPlaneScaleParameterSource::SetParameter(const FVector2D& NewValue)
 	if (bUseEqualScaling)
 	{
 		double SnappedDeltaX = 0.0, SnappedDeltaY = 0.0;
+		bool bIsSnapped = true;
 		
 		// if using snapping while scaling on X and Y axis
 		if (ScaleAxisXDeltaConstraintFunction(UseScaleDeltaX, SnappedDeltaX))
 		{
 			UseScaleDeltaX = SnappedDeltaX;
 		}
+		else
+		{
+			bIsSnapped = false;
+		}
 		if (ScaleAxisYDeltaConstraintFunction(UseScaleDeltaY, SnappedDeltaY))
 		{
 			UseScaleDeltaY = SnappedDeltaY;
 		}
-		// ensures that 2 axes on plane scale proportionally while following closest snap factor
-		// ex: Scale Snap is set to 1, StartScale = (1,2,3), scaling on X axis NewScale=(1, 4, 6), instead of (1, 3, 4)
-		NewScale = StartScale + (StartScale*(UseScaleDeltaX*CurScaleAxisX)) + (StartScale*(UseScaleDeltaY*CurScaleAxisY));
+		else
+		{
+			bIsSnapped = false;
+		}
+
+		// determines if the initial scales of the 2 affected axes are equivalent, and if we can therefore use uniform scaling (additive function)
+		const FVector AffectedValuesVector  = StartScale*CurScaleAxisX + StartScale*CurScaleAxisY;
+		const bool bIsUniformAcrossScaleAxes = (AffectedValuesVector.X == AffectedValuesVector.Y) || (AffectedValuesVector.X == AffectedValuesVector.Z) || (AffectedValuesVector.Y == AffectedValuesVector.Z);
+
+		// will use additive if scale is uniform across 2 axes of the plane AND snapping is on
+		if (bIsUniformAcrossScaleAxes && bIsSnapped)
+		{
+			// uses an additive function to scale when both initial values are equal
+			// ex: allows for a case where InitScale= (2,2,1) scaling by 1 across Z axis, next increment will be (3,3,1) instead of (4,4,1)
+			NewScale = StartScale + UseScaleDeltaX*CurScaleAxisX + UseScaleDeltaY*CurScaleAxisY;
+		}
+		else
+		{
+			NewScale = StartScale + (UseScaleDeltaX * StartScale * CurScaleAxisX) + (UseScaleDeltaY * StartScale * CurScaleAxisY);
+		}
 	}
 	else
 	{
-		NewScale = StartScale + ScaleDelta.X*CurScaleAxisX + ScaleDelta.Y*CurScaleAxisY;
+		NewScale = StartScale + (UseScaleDeltaX * StartScale * CurScaleAxisX) + (UseScaleDeltaY * StartScale * CurScaleAxisY);
 	}
 	
 	// currently calling ScaleConstraintFunction has no effect (no changes made to SignedDelta) because this constraint function
