ScriptableTools: Added support for two new input behaviors, Single Key Input and Multi Key Input, which allow ScriptableTool authors to capture events related to keyboard input, either as single keys or mutliple key combinations.

Includes the addition of a lower level KeyInputBehavior to the InteractiveToolsFramework to support these two behaviors.

diff --git a/Engine/Plugins/Runtime/ScriptableToolsFramework/Source/ScriptableToolsFramework/Private/BaseTools/ScriptableModularBehaviorTool.cpp b/Engine/Plugins/Runtime/ScriptableToolsFramework/Source/ScriptableToolsFramework/Private/BaseTools/ScriptableModularBehaviorTool.cpp
--- a/Engine/Plugins/Runtime/ScriptableToolsFramework/Source/ScriptableToolsFramework/Private/BaseTools/ScriptableModularBehaviorTool.cpp
+++ b/Engine/Plugins/Runtime/ScriptableToolsFramework/Source/ScriptableToolsFramework/Private/BaseTools/ScriptableModularBehaviorTool.cpp
@@ -9,6 +9,7 @@
 #include "Behaviors/ScriptableToolMouseWheelBehavior.h"
 #include "Behaviors/ScriptableToolMultiClickSequenceBehavior.h"
 #include "Behaviors/ScriptableToolSingleClickOrDragBehavior.h"
+#include "Behaviors/ScriptableToolKeyInputBehavior.h"
 
 void UScriptableModularBehaviorTool::AddSingleClickBehavior(
 	const FTestIfHitByClickDelegate TestIfHitByClickDelegate,
@@ -147,6 +148,41 @@ void UScriptableModularBehaviorTool::AddMouseHoverBehavior(
 }
 
 
+void UScriptableModularBehaviorTool::AddSingleKeyInputBehavior(
+	FOnKeyStateToggleDelegate OnKeyPressedFuncIn,
+	FOnKeyStateToggleDelegate OnKeyReleasedFuncIn,
+	FKey Key,
+	const FMouseBehaviorModiferCheckDelegate ModifierCheckFunction,
+	int CapturePriority
+)
+{
+	TArray<FKey> Keys;
+	Keys.Add(Key);
+
+	TObjectPtr<UScriptableToolKeyInputBehavior> BehaviorContainer = NewObject<UScriptableToolKeyInputBehavior>();
+	BehaviorContainer->Init(this, ModifierCheckFunction, OnKeyPressedFuncIn, OnKeyReleasedFuncIn, Keys, true);
+	BehaviorContainer->SetDefaultPriority(FInputCapturePriority(CapturePriority));
+
+	KeyInputBehaviors.Add(BehaviorContainer);
+}
+
+
+void UScriptableModularBehaviorTool::AddMultiKeyInputBehavior(
+	FOnKeyStateToggleDelegate OnKeyPressedFuncIn,
+	FOnKeyStateToggleDelegate OnKeyReleasedFuncIn,
+	TArray<FKey> Keys,
+	bool bRequireAllKeys,
+	const FMouseBehaviorModiferCheckDelegate ModifierCheckFunction,
+	int CapturePriority
+)
+{
+	TObjectPtr<UScriptableToolKeyInputBehavior> BehaviorContainer = NewObject<UScriptableToolKeyInputBehavior>();
+	BehaviorContainer->Init(this, ModifierCheckFunction, OnKeyPressedFuncIn, OnKeyReleasedFuncIn, Keys, bRequireAllKeys);
+	BehaviorContainer->SetDefaultPriority(FInputCapturePriority(CapturePriority));
+
+	KeyInputBehaviors.Add(BehaviorContainer);
+}
+
 void UScriptableModularBehaviorTool::OnUpdateModifierState(int ModifierID, bool bIsOn)
 {
 	if (ModifierID == 1)
diff --git a/Engine/Plugins/Runtime/ScriptableToolsFramework/Source/ScriptableToolsFramework/Private/Behaviors/ScriptableToolKeyInputBehavior.cpp b/Engine/Plugins/Runtime/ScriptableToolsFramework/Source/ScriptableToolsFramework/Private/Behaviors/ScriptableToolKeyInputBehavior.cpp
new file mode 100644
--- /dev/null
+++ b/Engine/Plugins/Runtime/ScriptableToolsFramework/Source/ScriptableToolsFramework/Private/Behaviors/ScriptableToolKeyInputBehavior.cpp
@@ -0,0 +1,74 @@
+// Copyright Epic Games, Inc. All Rights Reserved.
+
+#include "Behaviors/ScriptableToolKeyInputBehavior.h"
+
+#include "BaseTools/ScriptableModularBehaviorTool.h"
+#include "BaseBehaviors/KeyInputBehavior.h"
+
+void UScriptableToolKeyInputBehavior::Init(
+	TObjectPtr<UScriptableModularBehaviorTool> BehaviorHostIn,
+	FMouseBehaviorModiferCheckDelegate ModifierCheckFuncIn,
+	FOnKeyStateToggleDelegate OnKeyPressedFuncIn,
+	FOnKeyStateToggleDelegate OnKeyReleasedFuncIn,
+	const TArray<FKey>& ListenKeysIn,
+	bool bRequireAllKeysIn
+)
+{
+	BehaviorHost = BehaviorHostIn;
+	Behavior = CreateNewBehavior();
+	ModifierCheckFunc = ModifierCheckFuncIn;
+	OnKeyPressedFunc = OnKeyPressedFuncIn;
+	OnKeyReleasedFunc = OnKeyReleasedFuncIn;
+	ListenKeys = ListenKeysIn;
+
+	Behavior->Initialize(this, ListenKeys);
+	Behavior->ModifierCheckFunc = [this](const FInputDeviceState& InputDeviceState) {
+		bool bResult = true;
+		if (ModifierCheckFunc.IsBound())
+		{
+			bResult = ModifierCheckFunc.Execute(InputDeviceState);
+		}
+		return bResult;
+	};
+	Behavior->bRequireAllKeys = bRequireAllKeysIn;
+
+	BehaviorHost->AddInputBehavior(Behavior);
+
+	Behavior->Modifiers.RegisterModifier(1, FInputDeviceState::IsShiftKeyDown);
+	Behavior->Modifiers.RegisterModifier(2, FInputDeviceState::IsCtrlKeyDown);
+	Behavior->Modifiers.RegisterModifier(3, FInputDeviceState::IsAltKeyDown);
+
+}
+
+UKeyInputBehavior* UScriptableToolKeyInputBehavior::CreateNewBehavior() const
+{
+	return NewObject<UKeyInputBehavior>();
+}
+
+UInputBehavior* UScriptableToolKeyInputBehavior::GetWrappedBehavior()
+{
+	return Behavior;
+}
+
+void UScriptableToolKeyInputBehavior::OnKeyPressed(const FKey& Key)
+{
+	if (OnKeyPressedFunc.IsBound())
+	{
+		OnKeyPressedFunc.Execute(Key, BehaviorHost->GetActiveModifiers());
+	}
+	return;
+}
+
+void UScriptableToolKeyInputBehavior::OnKeyReleased(const FKey& Key)
+{
+	if (OnKeyReleasedFunc.IsBound())
+	{
+		OnKeyReleasedFunc.Execute(Key, BehaviorHost->GetActiveModifiers());
+	}
+	return;
+}
+
+void UScriptableToolKeyInputBehavior::OnUpdateModifierState(int ModifierID, bool bIsOn)
+{
+	BehaviorHost->OnUpdateModifierState(ModifierID, bIsOn);
+}
diff --git a/Engine/Plugins/Runtime/ScriptableToolsFramework/Source/ScriptableToolsFramework/Private/Behaviors/ScriptableToolKeyInputBehavior.h b/Engine/Plugins/Runtime/ScriptableToolsFramework/Source/ScriptableToolsFramework/Private/Behaviors/ScriptableToolKeyInputBehavior.h
new file mode 100644
--- /dev/null
+++ b/Engine/Plugins/Runtime/ScriptableToolsFramework/Source/ScriptableToolsFramework/Private/Behaviors/ScriptableToolKeyInputBehavior.h
@@ -0,0 +1,56 @@
+// Copyright Epic Games, Inc. All Rights Reserved.
+
+#pragma once
+
+#include "Behaviors/ScriptableToolBehaviorDelegates.h"
+#include "ScriptableToolBehavior.h"
+#include "BaseBehaviors/BehaviorTargetInterfaces.h"
+
+#include "ScriptableToolKeyInputBehavior.generated.h"
+
+class UScriptableModularBehaviorTool;
+class UKeyInputBehavior;
+
+UCLASS()
+class UScriptableToolKeyInputBehavior : public UScriptableToolBehavior, public IKeyInputBehaviorTarget
+{
+	GENERATED_BODY()
+
+public:
+	UScriptableToolKeyInputBehavior() {};
+
+	void Init(TObjectPtr<UScriptableModularBehaviorTool> BehaviorHostIn,
+		FMouseBehaviorModiferCheckDelegate ModifierCheckFuncIn,
+		FOnKeyStateToggleDelegate OnKeyPressedFuncIn,
+		FOnKeyStateToggleDelegate OnKeyReleasedFuncIn,
+		const TArray<FKey>& ListenKeysIn,
+		bool bRequireAllKeys);
+
+	virtual UKeyInputBehavior* CreateNewBehavior() const;
+	virtual UInputBehavior* GetWrappedBehavior() override;
+
+	virtual void OnKeyPressed(const FKey& KeyID) override;
+	virtual void OnKeyReleased(const FKey& KeyID) override;
+
+	virtual void OnUpdateModifierState(int ModifierID, bool bIsOn) override;
+
+private:
+
+	UPROPERTY()
+	TObjectPtr<UScriptableModularBehaviorTool> BehaviorHost;
+
+	UPROPERTY()
+	TObjectPtr<UKeyInputBehavior> Behavior;
+
+	UPROPERTY()
+	FMouseBehaviorModiferCheckDelegate ModifierCheckFunc;
+
+	UPROPERTY()
+	FOnKeyStateToggleDelegate	OnKeyPressedFunc;
+
+	UPROPERTY()
+	FOnKeyStateToggleDelegate OnKeyReleasedFunc;
+
+	UPROPERTY()
+	TArray<FKey> ListenKeys;
+};
diff --git a/Engine/Plugins/Runtime/ScriptableToolsFramework/Source/ScriptableToolsFramework/Public/BaseTools/ScriptableModularBehaviorTool.h b/Engine/Plugins/Runtime/ScriptableToolsFramework/Source/ScriptableToolsFramework/Public/BaseTools/ScriptableModularBehaviorTool.h
--- a/Engine/Plugins/Runtime/ScriptableToolsFramework/Source/ScriptableToolsFramework/Public/BaseTools/ScriptableModularBehaviorTool.h
+++ b/Engine/Plugins/Runtime/ScriptableToolsFramework/Source/ScriptableToolsFramework/Public/BaseTools/ScriptableModularBehaviorTool.h
@@ -15,6 +15,7 @@ class UScriptableToolSingleClickOrDragBehavior;
 class UScriptableToolMouseWheelBehavior;
 class UScriptableToolClickSequenceBehavior;
 class UScriptableToolMouseHoverBehavior;
+class UScriptableToolKeyInputBehavior;
 
 /**
  A Scriptable tool base blueprint class which provides support for user defined mouse interaction behaviors
@@ -225,6 +226,47 @@ public:
 		int CapturePriority = 100
 	);
 
+	/**
+	*	AddSingleKeyInputBehavior implements a generic keyboard key listener behavior
+	*	
+	*	@param OnKeyPressed Callback when the target key is pressed.
+	*	@param OnKeyReleased Callback when the target key is released
+	*	@param Key,Target key to watch for
+	*	@param CaptureCheck Only enable capture if returns true
+	*	@param CapturePriority The priority is used to resolve situations where multiple behaviors want the same capture
+	*/
+
+
+	UFUNCTION(BlueprintCallable, Category = "ScriptableTool|Input", meta = (AdvancedDisplay = "CaptureCheck, CapturePriority"))
+	void AddSingleKeyInputBehavior(
+		FOnKeyStateToggleDelegate OnKeyPressed,
+		FOnKeyStateToggleDelegate OnKeyReleased,
+		FKey Key,
+		const FMouseBehaviorModiferCheckDelegate CaptureCheck,
+		int CapturePriority
+	);
+
+	/**
+	*	AddMultiKeyInputBehavior implements a generic keyboard multi key listener behavior
+	*	
+	*	@param OnKeyPressed  Callback when the target key(s) is pressed. Only triggers once if bRequireAllKeys is true.
+	*	@param OnKeyReleased  Callback when the target key(s) is pressed. Only triggers once if bRequireAllKeys is true.
+	*	@param Keys Target keys to watch for
+	*	@param bRequireAllKeys If true, all target keys must be pressed simultaniously to recieve press/release events. Otherwise, any and all keys can trigger events.
+	*	@param CaptureCheck Only enable capture if returns true
+	*	@param CapturePriority The priority is used to resolve situations where multiple behaviors want the same capture
+	*/
+
+
+	UFUNCTION(BlueprintCallable, Category = "ScriptableTool|Input", meta = (AdvancedDisplay = "CaptureCheck, CapturePriority"))
+		void AddMultiKeyInputBehavior(
+			FOnKeyStateToggleDelegate OnKeyPressed,
+			FOnKeyStateToggleDelegate OnKeyReleased,
+			TArray<FKey> Keys,
+			bool bRequireAllKeys,
+			const FMouseBehaviorModiferCheckDelegate CaptureCheck,
+			int CapturePriority
+		);
 
 private:
 
@@ -249,6 +291,9 @@ private:
 	UPROPERTY(Transient, DuplicateTransient, NonTransactional, SkipSerialization)
 	TArray< TObjectPtr<UScriptableToolMouseHoverBehavior> > MouseHoverBehaviors;
 
+	UPROPERTY(Transient, DuplicateTransient, NonTransactional, SkipSerialization)
+	TArray< TObjectPtr<UScriptableToolKeyInputBehavior> > KeyInputBehaviors;
+
 
 	//
 	// Modifer Buttons Support
diff --git a/Engine/Plugins/Runtime/ScriptableToolsFramework/Source/ScriptableToolsFramework/Public/Behaviors/ScriptableToolBehaviorDelegates.h b/Engine/Plugins/Runtime/ScriptableToolsFramework/Source/ScriptableToolsFramework/Public/Behaviors/ScriptableToolBehaviorDelegates.h
--- a/Engine/Plugins/Runtime/ScriptableToolsFramework/Source/ScriptableToolsFramework/Public/Behaviors/ScriptableToolBehaviorDelegates.h
+++ b/Engine/Plugins/Runtime/ScriptableToolsFramework/Source/ScriptableToolsFramework/Public/Behaviors/ScriptableToolBehaviorDelegates.h
@@ -41,4 +41,8 @@ DECLARE_DYNAMIC_DELEGATE_RetVal(bool, FRequestAbortClickSequenceDelegate);
 DECLARE_DYNAMIC_DELEGATE_RetVal_TwoParams(FInputRayHit, FBeginHoverSequenceHitTestDelegate, FInputDeviceRay, CurrentPos, FScriptableToolModifierStates, Modifiers);
 DECLARE_DYNAMIC_DELEGATE_TwoParams(FOnBeginHoverDelegate, FInputDeviceRay, CurrentPos, FScriptableToolModifierStates, Modifiers);
 DECLARE_DYNAMIC_DELEGATE_RetVal_TwoParams(bool, FOnUpdateHoverDelegate, FInputDeviceRay, CurrentPos, FScriptableToolModifierStates, Modifiers);
-DECLARE_DYNAMIC_DELEGATE(FOnEndHoverDelegate);
\ No newline at end of file
+DECLARE_DYNAMIC_DELEGATE(FOnEndHoverDelegate);
+
+DECLARE_DYNAMIC_DELEGATE_TwoParams(FOnKeyStateToggleDelegate, FKey, Key, FScriptableToolModifierStates, Modifiers);
+
+
diff --git a/Engine/Source/Runtime/InteractiveToolsFramework/Private/BaseBehaviors/KeyInputBehavior.cpp b/Engine/Source/Runtime/InteractiveToolsFramework/Private/BaseBehaviors/KeyInputBehavior.cpp
new file mode 100644
--- /dev/null
+++ b/Engine/Source/Runtime/InteractiveToolsFramework/Private/BaseBehaviors/KeyInputBehavior.cpp
@@ -0,0 +1,150 @@
+// Copyright Epic Games, Inc. All Rights Reserved.
+
+#include "BaseBehaviors/KeyInputBehavior.h"
+#include "Algo/Unique.h"
+
+#include UE_INLINE_GENERATED_CPP_BY_NAME(KeyInputBehavior)
+
+UKeyInputBehavior::UKeyInputBehavior()
+{
+}
+
+void UKeyInputBehavior::Initialize(IKeyInputBehaviorTarget* TargetIn, const FKey& KeyIn)
+{
+	this->Target = TargetIn;
+	this->TargetKeys.Add(KeyIn);
+	InitializeKeyActivations();
+	
+}
+
+void UKeyInputBehavior::Initialize(IKeyInputBehaviorTarget* TargetIn, const TArray<FKey>& KeysIn)
+{
+	Target = TargetIn;
+
+	TargetKeys = KeysIn;
+	TargetKeys.Sort();
+	TargetKeys.SetNum(Algo::Unique(TargetKeys)); // Clean up array in case we have duplicate keys
+
+	InitializeKeyActivations();
+}
+
+FInputCaptureRequest UKeyInputBehavior::WantsCapture(const FInputDeviceState& Input)
+{
+	InitializeKeyActivations();
+	
+	if (IsTargetedKey(Input) && (ModifierCheckFunc == nullptr || ModifierCheckFunc(Input)))
+	{
+		return FInputCaptureRequest::Begin(this, EInputCaptureSide::Any);
+	}
+	return FInputCaptureRequest::Ignore();
+}
+
+
+FInputCaptureUpdate UKeyInputBehavior::BeginCapture(const FInputDeviceState& Input, EInputCaptureSide Side)
+{
+	Modifiers.UpdateModifiers(Input, Target);
+	UpdateActivations(Input, !bRequireAllKeys);
+
+	if ((bRequireAllKeys && AreAllKeysPressed()))
+	{
+		Target->OnKeyPressed(Input.Keyboard.ActiveKey.Button);
+		bAllKeysSeenPressed = true;
+	}
+
+	return FInputCaptureUpdate::Begin(this, EInputCaptureSide::Any);
+}
+
+
+FInputCaptureUpdate UKeyInputBehavior::UpdateCapture(const FInputDeviceState& Input, const FInputCaptureData& Data)
+{
+	Modifiers.UpdateModifiers(Input, Target);
+	bool bRelevantKey = UpdateActivations(Input, !bRequireAllKeys);
+
+	if (!bRelevantKey)
+	{
+		return FInputCaptureUpdate::Continue();
+	}
+
+	if (bRequireAllKeys && AreAllKeysPressed())
+	{
+		Target->OnKeyPressed(Input.Keyboard.ActiveKey.Button);
+		bAllKeysSeenPressed = true;
+		return FInputCaptureUpdate::Continue();
+	}
+
+	if (bRequireAllKeys && (!AreAllKeysPressed() && bAllKeysSeenPressed) )
+	{
+		Target->OnKeyReleased(Input.Keyboard.ActiveKey.Button);
+		return FInputCaptureUpdate::End();
+	}
+
+	if (!bRequireAllKeys && !IsAnyKeyPressed())
+	{
+		return FInputCaptureUpdate::End();
+	}
+
+	return FInputCaptureUpdate::Continue();
+}
+
+void UKeyInputBehavior::ForceEndCapture(const FInputCaptureData& Data)
+{
+}
+
+void UKeyInputBehavior::InitializeKeyActivations()
+{
+	KeyActivations.SetNum(TargetKeys.Num());
+	for (bool& Activation : KeyActivations)
+	{
+		Activation = false;
+	}
+}
+
+bool UKeyInputBehavior::IsTargetedKey(const FInputDeviceState& Input)
+{
+	return (Input.InputDevice == EInputDevices::Keyboard) &&
+		   (Input.Keyboard.ActiveKey.bPressed) &&
+		   (TargetKeys.Contains(Input.Keyboard.ActiveKey.Button));
+}
+
+bool UKeyInputBehavior::UpdateActivations(const FInputDeviceState& Input, bool bEmitOnChange)
+{
+	if ((Input.InputDevice == EInputDevices::Keyboard) && (Input.Keyboard.ActiveKey.bPressed))
+	{
+		int32 KeyIndex = TargetKeys.Find(Input.Keyboard.ActiveKey.Button);
+		if (KeyIndex > -1)
+		{
+			KeyActivations[KeyIndex] = true;
+			if (bEmitOnChange)
+			{
+				Target->OnKeyPressed(Input.Keyboard.ActiveKey.Button);
+			}
+			return true;
+		}
+	}
+	if ((Input.InputDevice == EInputDevices::Keyboard) && (Input.Keyboard.ActiveKey.bReleased))
+	{
+		int32 KeyIndex = TargetKeys.Find(Input.Keyboard.ActiveKey.Button);
+		if (KeyIndex > -1)
+		{
+			KeyActivations[KeyIndex] = false;
+			if (bEmitOnChange)
+			{
+				Target->OnKeyReleased(Input.Keyboard.ActiveKey.Button);
+			}
+			return true;
+		}
+	}
+
+	return false;
+}
+
+bool UKeyInputBehavior::IsAnyKeyPressed() const
+{
+	return KeyActivations.Contains(true);
+}
+
+bool UKeyInputBehavior::AreAllKeysPressed() const
+{
+	return !KeyActivations.Contains(false);
+}
+
diff --git a/Engine/Source/Runtime/InteractiveToolsFramework/Public/BaseBehaviors/BehaviorTargetInterfaces.h b/Engine/Source/Runtime/InteractiveToolsFramework/Public/BaseBehaviors/BehaviorTargetInterfaces.h
--- a/Engine/Source/Runtime/InteractiveToolsFramework/Public/BaseBehaviors/BehaviorTargetInterfaces.h
+++ b/Engine/Source/Runtime/InteractiveToolsFramework/Public/BaseBehaviors/BehaviorTargetInterfaces.h
@@ -26,6 +26,24 @@ public:
 	}
 };
 
+/**
+ * IKeyInputBehaviorTarget is an interface that InputBehaviors can use to notify
+ * a target about keyboard buttons toggling states, such as pressed or released.
+ */
+class IKeyInputBehaviorTarget : public IModifierToggleBehaviorTarget
+{
+public:
+	virtual ~IKeyInputBehaviorTarget() {}
+
+	virtual void OnKeyPressed(const FKey& KeyID)
+	{
+	}
+
+	virtual void OnKeyReleased(const FKey& KeyID)
+	{
+	}
+};
+
 
 /**
  * Functions required to apply standard "Click" state machines to a target object.
diff --git a/Engine/Source/Runtime/InteractiveToolsFramework/Public/BaseBehaviors/KeyInputBehavior.h b/Engine/Source/Runtime/InteractiveToolsFramework/Public/BaseBehaviors/KeyInputBehavior.h
new file mode 100644
--- /dev/null
+++ b/Engine/Source/Runtime/InteractiveToolsFramework/Public/BaseBehaviors/KeyInputBehavior.h
@@ -0,0 +1,108 @@
+// Copyright Epic Games, Inc. All Rights Reserved.
+
+#pragma once
+
+#include "BaseBehaviors/BehaviorTargetInterfaces.h"
+#include "BaseBehaviors/InputBehaviorModifierStates.h"
+#include "CoreMinimal.h"
+#include "InputBehavior.h"
+#include "InputCoreTypes.h"
+#include "InputState.h"
+#include "InteractiveTool.h"
+#include "Templates/Function.h"
+#include "UObject/ObjectMacros.h"
+#include "UObject/UObjectGlobals.h"
+
+#include "KeyInputBehavior.generated.h"
+
+class IModifierToggleBehaviorTarget;
+class UObject;
+
+
+/**
+ * UKeyInputBehavior provides a keyboard capturing behavior that works with single or multiple keys. Provided with a single key, the behavior triggers
+ * OnKeyPressed and OnKeyReleased events upon seeing the target key pressed down and released for the first time, ignoring any other key presses.
+ * When provided with multiple keys, the Behavior has variable behavior depending on whether bRequireAllKeys is set.
+ * 
+ * If true, the behavior sequence is as follows:
+ * 
+ *   1. Initiate capture when any of the target keys are pressed.
+ *   2. Continue capture until all target keys are pressed simultaneously
+ *   3. Upon seeing the last key to complete the full set of target keys, issue an OnKeyPressed for whichever key completed the requirement.
+ *		3b. If any of the target keys are released after the full set was pressed, issue an OnKeyReleased for whichever key was released, then end Capture.
+ *   4. If at any point all target keys are released after capture begins, end capture. 
+ * 
+ * If false, the behavior sequence is as follows:
+ * 
+ *   1. Initiate capture when any of the target keys are pressed.
+ *   2. Continue capture while any of the target keys are still pressed.
+ *   3. Issue an OnKeyPressed for any target key pressed during the capture period.
+ *   4. Issue an OnKeyReleased for any target key released during the capture period.
+ *   5. If at any point all target keys are released, end capture.
+ * 
+ */
+UCLASS(MinimalAPI)
+class UKeyInputBehavior : public UInputBehavior
+{
+	GENERATED_BODY()
+
+public:
+	INTERACTIVETOOLSFRAMEWORK_API UKeyInputBehavior();
+
+	virtual EInputDevices GetSupportedDevices() override
+	{
+		return EInputDevices::Keyboard;
+	}
+
+	/**
+	 * Initialize this behavior with the given Target
+	 * @param Target implementor of modifier-toggle behavior
+	 * @param Key FKey that indicates the keyboard key for the behavior to watch for	 
+	 */
+	INTERACTIVETOOLSFRAMEWORK_API virtual void Initialize(IKeyInputBehaviorTarget* Target, const FKey& Key);
+
+	/**
+	 * Initialize this behavior with the given Target
+	 * @param Target implementor of modifier-toggle behavior
+	 * @param Keys array of FKey that indicates the keyboard keys for the behavior to watch for	 
+	 */
+	INTERACTIVETOOLSFRAMEWORK_API virtual void Initialize(IKeyInputBehaviorTarget* Target, const TArray<FKey>& Keys);
+
+	/**
+	 * WantsCapture() will only return capture request if this function returns true (or is null)
+	 */
+	TFunction<bool(const FInputDeviceState&)> ModifierCheckFunc = nullptr;
+
+	// UInputBehavior implementation
+	INTERACTIVETOOLSFRAMEWORK_API virtual FInputCaptureRequest WantsCapture(const FInputDeviceState& Input) override;
+	INTERACTIVETOOLSFRAMEWORK_API virtual FInputCaptureUpdate BeginCapture(const FInputDeviceState& Input, EInputCaptureSide Side) override;
+	INTERACTIVETOOLSFRAMEWORK_API virtual FInputCaptureUpdate UpdateCapture(const FInputDeviceState& Input, const FInputCaptureData& Data) override;
+	INTERACTIVETOOLSFRAMEWORK_API virtual void ForceEndCapture(const FInputCaptureData& Data) override;
+
+
+	/**
+    * The modifier set for this behavior
+    */
+	FInputBehaviorModifierStates Modifiers;
+
+	/**
+	* If true, behavior requires all keys provided at initialization to be active at the same time in order to count as a "key down" event. If false, any key in the list will trigger the key down event.
+	*/
+	bool bRequireAllKeys = true;
+
+protected:
+	/** Modifier Target object */
+	IKeyInputBehaviorTarget* Target;
+
+	TArray<FKey> TargetKeys;
+	TArray<bool> KeyActivations;
+	
+private:
+	void InitializeKeyActivations();
+	bool IsTargetedKey(const FInputDeviceState& Input);
+	bool UpdateActivations(const FInputDeviceState& Input, bool bEmitOnChange);
+	bool IsAnyKeyPressed() const;
+	bool AreAllKeysPressed() const;
+
+	bool bAllKeysSeenPressed = false;
+};
