InteractiveToolsFramework: Add support for a new check function for determining if a hover component to a behavior can begin, similar to the modifer check function that currently exists for mouse click style behaviors.

See review for CL36131576

diff --git a/Engine/Source/Runtime/InteractiveToolsFramework/Private/BaseBehaviors/MouseHoverBehavior.cpp b/Engine/Source/Runtime/InteractiveToolsFramework/Private/BaseBehaviors/MouseHoverBehavior.cpp
--- a/Engine/Source/Runtime/InteractiveToolsFramework/Private/BaseBehaviors/MouseHoverBehavior.cpp
+++ b/Engine/Source/Runtime/InteractiveToolsFramework/Private/BaseBehaviors/MouseHoverBehavior.cpp
@@ -28,7 +28,7 @@ bool UMouseHoverBehavior::WantsHoverEvents()
 
 FInputCaptureRequest UMouseHoverBehavior::WantsHoverCapture(const FInputDeviceState& InputState)
 {
-	if (Target != nullptr)
+	if (Target != nullptr && ((HoverModifierCheckFunc == nullptr || HoverModifierCheckFunc(InputState))) )
 	{
 		Modifiers.UpdateModifiers(InputState, Target);
 
diff --git a/Engine/Source/Runtime/InteractiveToolsFramework/Private/BaseBehaviors/MultiClickSequenceInputBehavior.cpp b/Engine/Source/Runtime/InteractiveToolsFramework/Private/BaseBehaviors/MultiClickSequenceInputBehavior.cpp
--- a/Engine/Source/Runtime/InteractiveToolsFramework/Private/BaseBehaviors/MultiClickSequenceInputBehavior.cpp
+++ b/Engine/Source/Runtime/InteractiveToolsFramework/Private/BaseBehaviors/MultiClickSequenceInputBehavior.cpp
@@ -127,7 +127,7 @@ bool UMultiClickSequenceInputBehavior::WantsHoverEvents()
 
 FInputCaptureRequest UMultiClickSequenceInputBehavior::WantsHoverCapture(const FInputDeviceState& InputState)
 {
-	return Target != nullptr ? FInputCaptureRequest::Begin(this, EInputCaptureSide::Any) : FInputCaptureRequest::Ignore();
+	return (Target != nullptr) && ((HoverModifierCheckFunc == nullptr || HoverModifierCheckFunc(InputState))) ? FInputCaptureRequest::Begin(this, EInputCaptureSide::Any) : FInputCaptureRequest::Ignore();
 }
 
 FInputCaptureUpdate UMultiClickSequenceInputBehavior::BeginHoverCapture(const FInputDeviceState& InputState, EInputCaptureSide eSide) 
diff --git a/Engine/Source/Runtime/InteractiveToolsFramework/Public/BaseBehaviors/MouseHoverBehavior.h b/Engine/Source/Runtime/InteractiveToolsFramework/Public/BaseBehaviors/MouseHoverBehavior.h
--- a/Engine/Source/Runtime/InteractiveToolsFramework/Public/BaseBehaviors/MouseHoverBehavior.h
+++ b/Engine/Source/Runtime/InteractiveToolsFramework/Public/BaseBehaviors/MouseHoverBehavior.h
@@ -37,6 +37,15 @@ public:
 
 	INTERACTIVETOOLSFRAMEWORK_API virtual void Initialize(IHoverBehaviorTarget* Target);
 
+	/**
+	 * WantsHoverCapture() will only return capture request if this function returns true (or is null)
+	 *
+	 * Note that unlike click behaviors, hover behaviors continue to receive WantsCapture calls even when
+	 * they hold capture (and end if they fail to request capture), so an active hover behavior will end
+	 * if its HoverModifierCheckFunc returns false.
+	 */
+	TFunction<bool(const FInputDeviceState&)> HoverModifierCheckFunc = nullptr;
+
 	// UInputBehavior hover implementation
 
 	INTERACTIVETOOLSFRAMEWORK_API virtual EInputDevices GetSupportedDevices() override;
diff --git a/Engine/Source/Runtime/InteractiveToolsFramework/Public/BaseBehaviors/MultiClickSequenceInputBehavior.h b/Engine/Source/Runtime/InteractiveToolsFramework/Public/BaseBehaviors/MultiClickSequenceInputBehavior.h
--- a/Engine/Source/Runtime/InteractiveToolsFramework/Public/BaseBehaviors/MultiClickSequenceInputBehavior.h
+++ b/Engine/Source/Runtime/InteractiveToolsFramework/Public/BaseBehaviors/MultiClickSequenceInputBehavior.h
@@ -59,6 +59,14 @@ public:
 	 */
 	TFunction<bool(const FInputDeviceState&)> ModifierCheckFunc = nullptr;
 
+	/**
+	 * WantsHoverCapture() will only return capture request if this function returns true (or is null)
+	 *
+	 * Note that unlike click behaviors, hover behaviors continue to receive WantsCapture calls even when
+	 * they hold capture (and end if they fail to request capture), so an active hover behavior will end
+	 * if its HoverModifierCheckFunc returns false.
+	 */
+	TFunction<bool(const FInputDeviceState&)> HoverModifierCheckFunc = nullptr;
 
 	// UInputBehavior implementation
 
