InteractiveToolsFramework: Fix uninitialized variable complaint from CIS.

diff --git a/Engine/Source/Runtime/InteractiveToolsFramework/Private/BaseGizmos/GizmoArrowComponent.cpp b/Engine/Source/Runtime/InteractiveToolsFramework/Private/BaseGizmos/GizmoArrowComponent.cpp
--- a/Engine/Source/Runtime/InteractiveToolsFramework/Private/BaseGizmos/GizmoArrowComponent.cpp
+++ b/Engine/Source/Runtime/InteractiveToolsFramework/Private/BaseGizmos/GizmoArrowComponent.cpp
@@ -191,7 +191,7 @@ bool UGizmoArrowComponent::LineTraceComponent(FHitResult& OutHit, const FVector
 	FVector UseOrigin = Transform.TransformPosition(FVector::ZeroVector);
 
 	// Copy what is done in the proxy object, but get data from the gizmo view context.
-	FVector StartPoint, EndPoint;
+	FVector StartPoint = FVector::ZeroVector, EndPoint = FVector::ZeroVector; // initialized to appease CIS
 	float PixelToWorldScale = 0;
 	bool bRenderVisibility = GetWorldEndpoints(bIsViewDependent, ToRawPtr(GizmoViewContext), UseOrigin, Direction, Gap, Length, bWorld,
 		[&Transform](const FVector& VectorIn) { return Transform.TransformVector(VectorIn); },
