InteractiveToolsFramework: Move the gizmo component material, add an undimmed version (for free rotation handle), make the custom depth buffer comparison to scene depth in that material much more reliable. Also make CalculateViewDependentScaleAndFlatten usable both for FSceneView and GizmoViewContext, because it is more reliable than CalculateLocalPixelToWorldScale when zoomed far out in ortho.

#virtualized

diff --git a/Engine/Source/Runtime/InteractiveToolsFramework/Private/BaseGizmos/GizmoRenderingUtil.cpp b/Engine/Source/Runtime/InteractiveToolsFramework/Private/BaseGizmos/GizmoRenderingUtil.cpp
--- a/Engine/Source/Runtime/InteractiveToolsFramework/Private/BaseGizmos/GizmoRenderingUtil.cpp
+++ b/Engine/Source/Runtime/InteractiveToolsFramework/Private/BaseGizmos/GizmoRenderingUtil.cpp
@@ -106,9 +106,9 @@ float UE::GizmoRenderingUtil::CalculateLocalPixelToWorldScale(
 {
 	using namespace GizmoRenderingUtilLocals;
 
-	if (!View)
+	if (!ensure(View))
 	{
-		return 0.0f;
+		return 1.0f;
 	}
 
 	UE::GizmoRenderingUtil::FSceneViewWrapper Wrapper(*View);
@@ -145,12 +145,26 @@ float UE::GizmoRenderingUtil::CalculateViewDependentScaleAndFlatten(
 	const float InScale,
 	FVector& OutFlattenScale)
 {
-	const FMatrix& ViewMatrix = View->ViewMatrices.GetViewMatrix();
-	bool bIsPerspective = View->ViewMatrices.GetProjectionMatrix().M[3][3] < 1.0f;
+	if (!ensure(View))
+	{
+		return 1.0f;
+	}
+	UE::GizmoRenderingUtil::FSceneViewWrapper Wrapper(*View);
+	return CalculateViewDependentScaleAndFlatten(View, Location, InScale, OutFlattenScale);
+}
+
+float UE::GizmoRenderingUtil::CalculateViewDependentScaleAndFlatten(
+	const UE::GizmoRenderingUtil::ISceneViewInterface* View, 
+	const FVector& Location, 
+	const float InScale, 
+	FVector& OutFlattenScale)
+{
+	const FMatrix& ViewMatrix = View->GetViewMatrix();
+	bool bIsPerspective = View->GetProjectionMatrix().M[3][3] < 1.0f;
 	bool bIsOrthoXY = !bIsPerspective && FMath::Abs(ViewMatrix.M[2][2]) > 0.0f;
 	bool bIsOrthoXZ = !bIsPerspective && FMath::Abs(ViewMatrix.M[1][2]) > 0.0f;
 	bool bIsOrthoYZ = !bIsPerspective && FMath::Abs(ViewMatrix.M[0][2]) > 0.0f;
-	float UniformScale = static_cast<float> (InScale * View->WorldToScreen(Location).W * (4.0 / View->UnscaledViewRect.Width() / View->ViewMatrices.GetProjectionMatrix().M[0][0]));
+	float UniformScale = static_cast<float> (InScale * View->WorldToScreen(Location).W * (4.0 / View->GetUnscaledViewRect().Width() / View->GetProjectionMatrix().M[0][0]));
 
 	// Clamp to tolerance to prevent division by zero.
 	// @todo change to use MathUtil<RealType>::ZeroTolerance and TMathUtil<RealType>::SignNonZero(Value)
@@ -186,7 +200,13 @@ float UE::GizmoRenderingUtil::CalculateViewDependentScaleAndFlatten(
 
 UMaterialInterface* UE::GizmoRenderingUtil::GetDefaultGizmoComponentMaterial(const FLinearColor& Color, UObject* Outer)
 {
-	UMaterialInterface* Material = LoadObject<UMaterial>(nullptr, TEXT("/MeshModelingToolsetExp/Materials/GizmoComponentMaterial"));
+	return GetDefaultGizmoComponentMaterial(Color, FDefaultGizmoMaterialExtraParams(), Outer);
+}
+
+UMaterialInterface* UE::GizmoRenderingUtil::GetDefaultGizmoComponentMaterial(const FLinearColor & Color, const FDefaultGizmoMaterialExtraParams & Params, UObject * Outer)
+{
+	UMaterialInterface* Material = Params.bDimOccluded ? LoadObject<UMaterial>(nullptr, TEXT("/Engine/InteractiveToolsFramework/Materials/GizmoComponentMaterial"))
+		: LoadObject<UMaterial>(nullptr, TEXT("/Engine/InteractiveToolsFramework/Materials/GizmoComponentMaterial_NotDimmed"));
 	if (Material == nullptr)
 	{
 		return nullptr;
diff --git a/Engine/Source/Runtime/InteractiveToolsFramework/Public/BaseGizmos/GizmoRenderingUtil.h b/Engine/Source/Runtime/InteractiveToolsFramework/Public/BaseGizmos/GizmoRenderingUtil.h
--- a/Engine/Source/Runtime/InteractiveToolsFramework/Public/BaseGizmos/GizmoRenderingUtil.h
+++ b/Engine/Source/Runtime/InteractiveToolsFramework/Public/BaseGizmos/GizmoRenderingUtil.h
@@ -30,6 +30,7 @@ namespace UE::GizmoRenderingUtil
 		virtual FVector GetViewDirection() const = 0;
 		virtual FVector GetViewRight() const = 0;
 		virtual FVector GetViewUp() const = 0;
+		virtual const FMatrix& GetProjectionMatrix() const = 0;
 		virtual const FMatrix& GetViewMatrix() const = 0;
 		virtual bool IsPerspectiveProjection() const = 0;
 	};
@@ -49,6 +50,7 @@ namespace UE::GizmoRenderingUtil
 		virtual FVector GetViewDirection() const { return View->GetViewDirection(); }
 		virtual FVector GetViewRight() const { return View->GetViewRight(); }
 		virtual FVector GetViewUp() const { return View->GetViewUp(); }
+		virtual const FMatrix& GetProjectionMatrix() const { return View->ViewMatrices.GetProjectionMatrix(); }
 		virtual const FMatrix& GetViewMatrix() const { return View->ViewMatrices.GetViewMatrix(); }
 		virtual bool IsPerspectiveProjection() const { return View->IsPerspectiveProjection(); }
 
@@ -71,6 +73,15 @@ namespace UE::GizmoRenderingUtil
 	 */
 	INTERACTIVETOOLSFRAMEWORK_API UMaterialInterface* GetDefaultGizmoComponentMaterial(const FLinearColor& Color, UObject* Outer = nullptr);
 
+	struct FDefaultGizmoMaterialExtraParams 
+	{
+		bool bDimOccluded = true;
+	};
+	/** 
+	 * Like the other overload, but with additional parameters.
+	 */
+	UMaterialInterface* GetDefaultGizmoComponentMaterial(const FLinearColor& Color, const FDefaultGizmoMaterialExtraParams& Params, UObject* Outer = nullptr);
+
 	/**
 	 * Helper that creates a component with the default gizmo material and sets up the component-side properties that are needed for
 	 *  it to properly work (translucency sort priority, etc). This may not be necessary if your component is using some other
@@ -108,6 +119,11 @@ namespace UE::GizmoRenderingUtil
 	 * @return OutWorldFlattenScale vector to be applied in world space, can be used to flatten excluded 
 	 *         dimension in orthographic views as it reverses the scale in that dimension.
 	 */
+	INTERACTIVETOOLSFRAMEWORK_API float CalculateViewDependentScaleAndFlatten(
+		const UE::GizmoRenderingUtil::ISceneViewInterface* View,
+		const FVector& Location,
+		const float Scale,
+		FVector& OutWorldFlattenScale);
 	INTERACTIVETOOLSFRAMEWORK_API float CalculateViewDependentScaleAndFlatten(
 		const FSceneView* View,
 		const FVector& Location,
diff --git a/Engine/Source/Runtime/InteractiveToolsFramework/Public/BaseGizmos/GizmoViewContext.h b/Engine/Source/Runtime/InteractiveToolsFramework/Public/BaseGizmos/GizmoViewContext.h
--- a/Engine/Source/Runtime/InteractiveToolsFramework/Public/BaseGizmos/GizmoViewContext.h
+++ b/Engine/Source/Runtime/InteractiveToolsFramework/Public/BaseGizmos/GizmoViewContext.h
@@ -28,16 +28,19 @@ public:
 	public:
 		void ResetFromSceneView(const FSceneView& SceneView)
 		{
+			ProjectionMatrix = SceneView.ViewMatrices.GetProjectionMatrix();
 			ViewMatrix = SceneView.ViewMatrices.GetViewMatrix();
 			InvViewMatrix = SceneView.ViewMatrices.GetInvViewMatrix();
 			ViewProjectionMatrix = SceneView.ViewMatrices.GetViewProjectionMatrix();
 		}
 
+		const FMatrix& GetProjectionMatrix() const { return ProjectionMatrix; }
 		const FMatrix& GetViewMatrix() const { return ViewMatrix; }
 		const FMatrix& GetInvViewMatrix() const { return InvViewMatrix; }
 		const FMatrix& GetViewProjectionMatrix() const { return ViewProjectionMatrix; }
 
 	protected:
+		FMatrix ProjectionMatrix;
 		FMatrix ViewMatrix;
 		FMatrix InvViewMatrix;
 		FMatrix ViewProjectionMatrix;
@@ -58,6 +61,7 @@ public:
 	FVector GetViewRight() const override { return ViewMatrices.GetViewMatrix().GetColumn(0); }
 	FVector GetViewUp() const override { return ViewMatrices.GetViewMatrix().GetColumn(1); }
 	FVector GetViewDirection() const override { return ViewMatrices.GetViewMatrix().GetColumn(2); }
+	const FMatrix& GetProjectionMatrix() const override { return ViewMatrices.GetProjectionMatrix(); }
 	const FMatrix& GetViewMatrix() const override { return ViewMatrices.GetViewMatrix(); }
 	bool IsPerspectiveProjection() const override { return bIsPerspectiveProjection; }
 
