InteractiveToolsFramework: Changes to ViewAdjustedStaticMeshGizmoComponent:

- Make it capable of hiding itself based on view information (to hide axes while looking down them, etc).
- Allow transform adjusters to return their own view-independent bounds.
- Make transform adjusters able to keep their components in front of the ortho camera.

diff --git a/Engine/Source/Runtime/InteractiveToolsFramework/Private/BaseGizmos/ViewAdjustedStaticMeshGizmoComponent.cpp b/Engine/Source/Runtime/InteractiveToolsFramework/Private/BaseGizmos/ViewAdjustedStaticMeshGizmoComponent.cpp
--- a/Engine/Source/Runtime/InteractiveToolsFramework/Private/BaseGizmos/ViewAdjustedStaticMeshGizmoComponent.cpp
+++ b/Engine/Source/Runtime/InteractiveToolsFramework/Private/BaseGizmos/ViewAdjustedStaticMeshGizmoComponent.cpp
@@ -6,6 +6,8 @@
 #include "BaseGizmos/CombinedTransformGizmo.h"
 #include "BaseGizmos/ViewBasedTransformAdjusters.h"
 #include "BaseGizmos/GizmoRenderingUtil.h"
+#include "Materials/MaterialRenderProxy.h"
+#include "SceneInterface.h"
 #include "StaticMeshResources.h"
 #include "StaticMeshSceneProxy.h"
 #include "StaticMeshSceneProxyDesc.h"
@@ -44,6 +46,7 @@ namespace ViewAdjustedStaticMeshGizmoComponentLocals
 			, HoverOverrideMaterial(Component->GetHoverOverrideMaterial())
 			, bHovered(Component->IsBeingHovered())
 			, bHiddenByInteraction(Component->IsHiddenByInteraction())
+			, RenderVisibilityFunc(Component->GetRenderVisibilityFunction())
 		{
 		}
 
@@ -88,8 +91,15 @@ namespace ViewAdjustedStaticMeshGizmoComponentLocals
 				bool bAdjustedTransformDeterminantIsNegative = IsLocalToWorldDeterminantNegative();
 				if (TransformAdjuster.IsValid())
 				{
+					UE::GizmoRenderingUtil::FSceneViewWrapper WrappedView(*View);
+
 					FTransform NewTransform = TransformAdjuster->GetAdjustedComponentToWorld_RenderThread(
-						UE::GizmoRenderingUtil::FSceneViewWrapper(*View), FTransform(GetLocalToWorld()));
+						WrappedView, FTransform(GetLocalToWorld()));
+
+					if (RenderVisibilityFunc && !RenderVisibilityFunc(WrappedView, NewTransform))
+					{
+						continue;
+					}
 
 					FMatrix NewTransformMatrix = NewTransform.ToMatrixWithScale();
 					bAdjustedTransformDeterminantIsNegative = (NewTransformMatrix.Determinant() < 0);
@@ -106,6 +116,14 @@ namespace ViewAdjustedStaticMeshGizmoComponentLocals
 
 					AdjustedTransformBuffer = &DynamicPrimitiveUniformBuffer.UniformBuffer;
 				}//end getting adjusted transform ready
+				else
+				{
+					// If we're not adjusting the transform, just check visibility with the original transform
+					if (RenderVisibilityFunc && !RenderVisibilityFunc(UE::GizmoRenderingUtil::FSceneViewWrapper(*View), FTransform(GetLocalToWorld())))
+					{
+						continue;
+					}
+				}
 
 				FLODMask LODMask = GetLODMask(View);
 
@@ -126,13 +144,13 @@ namespace ViewAdjustedStaticMeshGizmoComponentLocals
 						{
 							FMeshBatch& MeshBatch = Collector.AllocateMesh();
 
+							const FLODInfo::FSectionInfo& Section = LODs[LODIndex].Sections[SectionIndex];
+
 							// This selection and hit proxy id setting seems unneeded but we'll just keep it
 							bool bSectionIsSelected = false;
 #if WITH_EDITOR
 							if (GIsEditor)
 							{
-								const FLODInfo::FSectionInfo& Section = LODs[LODIndex].Sections[SectionIndex];
-
 								bSectionIsSelected = Section.bSelected;
 								MeshBatch.BatchHitProxyId = Section.HitProxy ? Section.HitProxy->Id : FHitProxyId();
 							}
@@ -150,16 +168,31 @@ namespace ViewAdjustedStaticMeshGizmoComponentLocals
 							// Seems like there is only ever one of these...
 							FMeshBatchElement& BatchElement = MeshBatch.Elements[0];
 
-							// Update the portions affected by the adjusted transform determinant
-							const bool bUseReversedIndices =
-								// based off FStaticMeshSceneProxy::ShouldRenderBackFaces()
-								(bAdjustedTransformDeterminantIsNegative != bReverseCulling) &&
-								RenderData->LODResources[LODIndex].bHasReversedDepthOnlyIndices;
-							MeshBatch.ReverseCulling = IsReversedCullingNeeded(bUseReversedIndices);
+							// The above GetMeshElement does not reflect the adjusted transform determinant, so we have to
+							//  redo those pieces of MeshBatch setup.
+							
+							// Updated output of FStaticMeshSceneProxy::ShouldRenderBackFaces() 
+							bool bShouldRenderBackFaces = bReverseCulling != bAdjustedTransformDeterminantIsNegative;
+
+							// Updated value for bUseReversedIndices
+							UMaterialInterface* MaterialInterface = Section.Material;
+							const FMaterialRenderProxy* MaterialRenderProxy = MaterialInterface->GetRenderProxy();
+							const ERHIFeatureLevel::Type FeatureLevel = GetScene().GetFeatureLevel();
+							const FMaterial& Material = MaterialRenderProxy->GetIncompleteMaterialWithFallback(FeatureLevel);
+							const bool bUseReversedIndices = bShouldRenderBackFaces
+								&& RenderData->LODResources[LODIndex].bHasReversedDepthOnlyIndices
+								&& !Material.IsTwoSided();
+
+							// This is originally done in FStaticMeshSceneProxy::SetMeshElementGeometrySource
 							BatchElement.IndexBuffer = bUseReversedIndices 
 								? &LODModel.AdditionalIndexBuffers->ReversedIndexBuffer 
 								: &LODModel.IndexBuffer;
 
+							// Updated output of FStaticMeshSceneProxy::IsReversedCullingNeeded
+							MeshBatch.ReverseCulling = bShouldRenderBackFaces && !bUseReversedIndices;
+							// Done updating based on transform determinant
+
+
 							// Gizmos probably don't want to be affected by view mode?
 							MeshBatch.bCanApplyViewModeOverrides = false;
 
@@ -217,6 +250,9 @@ namespace ViewAdjustedStaticMeshGizmoComponentLocals
 		//  depending on the current gizmo mode (translate/rotate/scale). Instead, we want this setting to be an extra flag
 		//  that forces invisibility. 
 		bool bHiddenByInteraction = false;
+
+		TFunction<bool(const UE::GizmoRenderingUtil::ISceneViewInterface& View,
+			const FTransform& ComponentToWorld)> RenderVisibilityFunc = nullptr;
 	};
 
 }
@@ -224,21 +260,49 @@ namespace ViewAdjustedStaticMeshGizmoComponentLocals
 
 bool UViewAdjustedStaticMeshGizmoComponent::LineTraceComponent(FHitResult& OutHit, const FVector Start, const FVector End, const FCollisionQueryParams& Params)
 {
+	if (!ensure(GizmoViewContext))
+	{
+		return Super::LineTraceComponent(OutHit, Start, End, Params);
+	}
+
 	// If needed, update the physics data, then do the line trace
-	if (ensure(GizmoViewContext) && TransformAdjuster.IsValid())
+	if (TransformAdjuster.IsValid())
 	{
 		FTransform OriginalComponentToWorld = GetComponentToWorld();
 		FTransform AdjustedComponentToWorld = TransformAdjuster->GetAdjustedComponentToWorld(*GizmoViewContext, OriginalComponentToWorld);
+
+		if (RenderVisibilityFunc && !RenderVisibilityFunc(*GizmoViewContext, AdjustedComponentToWorld))
+		{
+			return false;
+		}
+
 		if (!OriginalComponentToWorld.Equals(AdjustedComponentToWorld))
 		{
 			BodyInstance.SetBodyTransform(AdjustedComponentToWorld, ETeleportType::None);
 			BodyInstance.UpdateBodyScale(AdjustedComponentToWorld.GetScale3D());
 		}
 	}
+	else
+	{
+		if (RenderVisibilityFunc && !RenderVisibilityFunc(*GizmoViewContext, GetComponentToWorld()))
+		{
+			return false;
+		}
+	}
 
 	return Super::LineTraceComponent(OutHit, Start, End, Params);
 }
 
+FBoxSphereBounds UViewAdjustedStaticMeshGizmoComponent::CalcBounds(const FTransform& LocalToWorld) const
+{
+	FBoxSphereBounds OriginalBounds = Super::CalcBounds(LocalToWorld);
+	if (!TransformAdjuster.IsValid())
+	{
+		return OriginalBounds;
+	}
+	return TransformAdjuster->GetViewIndependentBounds(LocalToWorld, OriginalBounds);
+}
+
 FPrimitiveSceneProxy* UViewAdjustedStaticMeshGizmoComponent::CreateStaticMeshSceneProxy(Nanite::FMaterialAudit& NaniteMaterials, bool bCreateNanite)
 {
 	return ::new ViewAdjustedStaticMeshGizmoComponentLocals::FViewAdjustedStaticMeshGizmoComponentProxy(this);
@@ -260,6 +324,13 @@ void UViewAdjustedStaticMeshGizmoComponent::SetTransformAdjuster(TSharedPtr<UE::
 	MarkRenderStateDirty();
 }
 
+void UViewAdjustedStaticMeshGizmoComponent::SetRenderVisibilityFunction(TFunction<bool(
+	const UE::GizmoRenderingUtil::ISceneViewInterface& View, const FTransform& ComponentToWorld)> RenderVisibilityFuncIn)
+{
+	RenderVisibilityFunc = RenderVisibilityFuncIn;
+	MarkRenderStateDirty();
+}
+
 void UViewAdjustedStaticMeshGizmoComponent::SetAllMaterials(UMaterialInterface* Material)
 {
 	for (int i = 0; i < GetNumMaterials(); ++i)
diff --git a/Engine/Source/Runtime/InteractiveToolsFramework/Private/BaseGizmos/ViewBasedTransformAdjusters.cpp b/Engine/Source/Runtime/InteractiveToolsFramework/Private/BaseGizmos/ViewBasedTransformAdjusters.cpp
--- a/Engine/Source/Runtime/InteractiveToolsFramework/Private/BaseGizmos/ViewBasedTransformAdjusters.cpp
+++ b/Engine/Source/Runtime/InteractiveToolsFramework/Private/BaseGizmos/ViewBasedTransformAdjusters.cpp
@@ -8,27 +8,39 @@
 
 namespace ViewBasedTransformAdjustersLocals
 {
+	FBoxSphereBounds InfiniteBounds(FSphere(FVector::ZeroVector, HALF_WORLD_MAX));
+
 	double GetComponentToGizmoScaling(
 		const UE::GizmoRenderingUtil::ISceneViewInterface& View,
 		const FTransform& GizmoToWorld)
 	{
-		// It might seem that all we want is CalculateLocalPixelToWorldScale(View, CurrentComponentToWorld.GetLocation());
+		// It might seem that all we want is to calculate local scale at CurrentComponentToWorld.GetLocation().
 		//  However we're in a weird situation where scaling gets applied around the gizmo origin, not the component (to
 		//  preserve component positioning within the gizmo); moreover we want different sub components of the same gizmo
-		//  to be scaled the same way. So, we need to use gizmo location as the basis of our scaling.
-		// On the other hand, we can't directly use CalculateLocalPixelToWorldScale with gizmo origin either because it's
-		//  possible to see the sub gizmo while the gizmo itself is off screen, breaking that calculation. What we want 
-		//  instead is a scale value that is dependent on FOV and distance to gizmo, but independent of view direction. 
-		//  The way we do this is we take the distance to gizmo, and then evaluate the scale at a point that far along the
-		//  view direction (so the actual view direction no longer matters- we're always looking at the scale around the
-		//  center of the screen).
+		//  to be scaled the same way. So, we need to use owning gizmo location as the basis of our scaling.
+		// On the other hand, we can't directly calculate local scale at gizmo origin either because it's possible to see
+		//  the sub gizmo while the gizmo itself is off screen, breaking that calculation. Moreover, we would like our 
+		//  scaling to stay the same as we move the gizmo in the view plane. Still, we want our scaling to account for FOV.
+		// The way we do this is we take the distance to gizmo origin, project that onto the view direction, and evaluate
+		//  the local scale to use there. So, the same scale is applied to all sub components of the same gizmo, this
+		//  scale is not affected by translations in the view plane, and the gizmo origin can be off screen.
 		// Note that this doesn't fully fix all problems with large component-to-gizmo offsets. Namely, moving the component
 		//  with the gizmo off screen can sometimes make it seem like the component stops moving and starts growing, which 
 		//  looks unnatural if you're not looking at the gizmo and realizing that you are just modifying your angle relative
 		//  to it. However the approach here tends to be the least broken overall.
-		double DistanceToGizmo = FVector3d::Distance(GizmoToWorld.GetLocation(), View.GetViewLocation());
-		FVector3d PointAtWhichToCheckScale = View.GetViewLocation() + DistanceToGizmo * View.GetViewDirection();
-		return UE::GizmoRenderingUtil::CalculateLocalPixelToWorldScale(&View, PointAtWhichToCheckScale);
+
+		double ProjectedDistanceToGizmo = (GizmoToWorld.GetLocation() - View.GetViewLocation()).Dot(View.GetViewDirection());
+		FVector3d PointAtWhichToCheckScale = View.GetViewLocation() + ProjectedDistanceToGizmo * View.GetViewDirection();
+		
+		// CalculateLocalPixelToWorldScale is inconsistent in ortho viewports when zoomed very far out, giving a scale
+		//  that seems to change at the same zoom level. Instead we use the same approach that editor gizmos use for getting
+		//  local scale.
+		FVector Unused;
+		return UE::GizmoRenderingUtil::CalculateViewDependentScaleAndFlatten(&View, PointAtWhichToCheckScale, 
+			// 0.5 scale here happens to get us close to what CalculateLocalPixelToWorldScale would return
+			//  at that same position.
+			0.5, 
+			Unused);
 	}
 
 	// Apply the settings to get a transform modified by view
@@ -38,7 +50,9 @@ namespace ViewBasedTransformAdjustersLocals
 		const FTransform& GizmoOriginToComponent,
 		const UE::GizmoRenderingUtil::FSubGizmoTransformAdjuster::FSettings& Settings)
 	{
-		if (!Settings.bKeepConstantViewSize && !Settings.bMirrorBasedOnOctant && !Settings.bUseWorldAxesForGizmo)
+		bool bIsOrtho = !View.IsPerspectiveProjection();
+		if (!Settings.bKeepConstantViewSize && !Settings.bMirrorBasedOnOctant && !Settings.bUseWorldAxesForGizmo 
+			&& (Settings.DistanceInFrontOfCameraInOrtho <= 0 || !bIsOrtho))
 		{
 			return CurrentComponentToWorld;
 		}
@@ -46,6 +60,17 @@ namespace ViewBasedTransformAdjustersLocals
 		FTransform GizmoToWorld = GizmoOriginToComponent * CurrentComponentToWorld;
 		FTransform ComponentToGizmo = CurrentComponentToWorld.GetRelativeTransform(GizmoToWorld);
 
+		if (bIsOrtho && Settings.DistanceInFrontOfCameraInOrtho >= 0)
+		{
+			UE::Geometry::FFrame3d CameraFrame(View.GetViewLocation(), View.GetViewDirection(), View.GetViewRight(), View.GetViewUp());
+			FVector CameraSpaceLocation = CameraFrame.ToFramePoint(GizmoToWorld.GetLocation());
+			if (CameraSpaceLocation.X < Settings.DistanceInFrontOfCameraInOrtho)
+			{
+				CameraSpaceLocation.X = Settings.DistanceInFrontOfCameraInOrtho;
+					GizmoToWorld.SetLocation(CameraFrame.FromFramePoint(CameraSpaceLocation));
+			}
+		}
+
 		if (Settings.bUseWorldAxesForGizmo)
 		{
 			GizmoToWorld.SetRotation(FQuat::Identity);
@@ -66,7 +91,7 @@ namespace ViewBasedTransformAdjustersLocals
 			EAxis::Type Axes[3] = { EAxis::X, EAxis::Y, EAxis::Z };
 			for (int Dim = 0; Dim < 3; ++Dim)
 			{
-				if (GizmoSpaceDirectionTowardCamera[Dim] < 0)
+				if (GizmoSpaceDirectionTowardCamera[Dim] < -KINDA_SMALL_NUMBER)
 				{
 					ComponentToGizmo.Mirror(Axes[Dim], EAxis::None);
 				}
@@ -97,6 +122,12 @@ FTransform UE::GizmoRenderingUtil::FSimpleConstantViewScaleAdjuster::GetAdjusted
 	return AdjustedTransform;
 }
 
+FBoxSphereBounds UE::GizmoRenderingUtil::FSimpleConstantViewScaleAdjuster::GetViewIndependentBounds(
+	const FTransform& LocalToWorld, const FBoxSphereBounds& OriginalBounds)
+{
+	return ViewBasedTransformAdjustersLocals::InfiniteBounds;
+}
+
 
 // FSubGizmoTransformAdjuster:
 
@@ -126,6 +157,10 @@ TSharedPtr<UE::GizmoRenderingUtil::FSubGizmoTransformAdjuster> UE::GizmoRenderin
 	UE::GizmoRenderingUtil::FSubGizmoTransformAdjuster::FSettings Settings;
 	Settings.bKeepConstantViewSize = true;
 	Settings.bMirrorBasedOnOctant = bMirrorBasedOnOctant;
+	// DistanceInFrontOfCameraInOrtho is used to keep the gizmo from being culled by the near or far plane in the ortho
+	//  viewports, by keeping it at a fixed distance in front of the camera. 1000 is an arbitrary value that's fine as long
+	//  as the gizmo is not super thick.
+	Settings.DistanceInFrontOfCameraInOrtho = 1000;
 	// bUseWorldAxes gets updated automatically as part of UViewAdjustedStaticMeshGizmoComponent::UpdateWorldLocalState
 	TransformAdjuster->SetSettings(Settings);
 	if (ComponentToKeepDistanceConstantTo)
@@ -152,6 +187,27 @@ FTransform UE::GizmoRenderingUtil::FSubGizmoTransformAdjuster::GetAdjustedCompon
 		GizmoOriginToComponent_RenderThread, Settings);
 }
 
+FBoxSphereBounds UE::GizmoRenderingUtil::FSubGizmoTransformAdjuster::GetViewIndependentBounds(
+	const FTransform& LocalToWorld, const FBoxSphereBounds& OriginalBounds)
+{
+	if (!Settings.bKeepConstantViewSize && !Settings.bMirrorBasedOnOctant && !Settings.bUseWorldAxesForGizmo)
+	{
+		return OriginalBounds;
+	}
+	if (Settings.bKeepConstantViewSize)
+	{
+		return ViewBasedTransformAdjustersLocals::InfiniteBounds;
+	}
+
+	// We could probably figure out how to make sure we only include the appropriate mirrored portions in our
+	//  bounds, but we'll just make it so that our bounds include any orientation relative to the gizmo origin.
+	FVector GizmoWorldLocation = LocalToWorld.TransformPosition(GizmoOriginToComponent_GameThread.GetLocation());
+	double DistanceToGizmoOriginWorld = (GizmoWorldLocation - OriginalBounds.Origin).Size();
+	double TotalRadius = DistanceToGizmoOriginWorld + FMath::Max(OriginalBounds.BoxExtent.Size(), OriginalBounds.SphereRadius);
+	
+	return FBoxSphereBounds(FSphere(OriginalBounds.Origin, TotalRadius));
+}
+
 
 // FConstantViewRelativeTransformAdjuster:
 
@@ -173,4 +229,17 @@ FTransform UE::GizmoRenderingUtil::FConstantViewRelativeTransformAdjuster::GetAd
 	}
 
 	return ViewRelativeTransformToUse * ViewFrameAtComponent.ToFTransform();
-}
\ No newline at end of file
+}
+
+FBoxSphereBounds UE::GizmoRenderingUtil::FConstantViewRelativeTransformAdjuster::GetViewIndependentBounds(
+	const FTransform& LocalToWorld, const FBoxSphereBounds& OriginalBounds)
+{
+	if (bKeepConstantViewSize)
+	{
+		return ViewBasedTransformAdjustersLocals::InfiniteBounds;
+	}
+
+	// Otherwise, the only thing that is changing is the orientation. Unfortunately we don't know if
+	//  the bounds started as a sphere or a box, so we'll choose the larger for our radius.
+	return FBoxSphereBounds(FSphere(OriginalBounds.Origin, FMath::Max(OriginalBounds.BoxExtent.Size(), OriginalBounds.SphereRadius)));
+}
diff --git a/Engine/Source/Runtime/InteractiveToolsFramework/Public/BaseGizmos/ViewAdjustedStaticMeshGizmoComponent.h b/Engine/Source/Runtime/InteractiveToolsFramework/Public/BaseGizmos/ViewAdjustedStaticMeshGizmoComponent.h
--- a/Engine/Source/Runtime/InteractiveToolsFramework/Public/BaseGizmos/ViewAdjustedStaticMeshGizmoComponent.h
+++ b/Engine/Source/Runtime/InteractiveToolsFramework/Public/BaseGizmos/ViewAdjustedStaticMeshGizmoComponent.h
@@ -11,6 +11,7 @@ class UGizmoViewContext;
 namespace UE::GizmoRenderingUtil
 {
 	class IViewBasedTransformAdjuster;
+	class ISceneViewInterface;
 }
 
 /**
@@ -30,7 +31,17 @@ public:
 	void SetGizmoViewContext(UGizmoViewContext* GizmoViewContextIn) { GizmoViewContext = GizmoViewContextIn; }
 	
 	void SetTransformAdjuster(TSharedPtr<UE::GizmoRenderingUtil::IViewBasedTransformAdjuster> Adjuster);
-	TSharedPtr<UE::GizmoRenderingUtil::IViewBasedTransformAdjuster> GetTransformAdjuster() { return TransformAdjuster; }
+	TSharedPtr<UE::GizmoRenderingUtil::IViewBasedTransformAdjuster> GetTransformAdjuster() const { return TransformAdjuster; }
+
+	/**
+	 * The render visibility function is an optional function that can hide the component based on view
+	 *  information (for instance to hide an arrow gizmo when looking directly down the arrow). It must 
+	 *  be callable from both the game and the render threads (for line traces and rendering).
+	 */
+	void SetRenderVisibilityFunction(TFunction<bool(const UE::GizmoRenderingUtil::ISceneViewInterface& View, 
+		const FTransform& ComponentToWorld)> RenderVisibilityFunc);
+	TFunction<bool(const UE::GizmoRenderingUtil::ISceneViewInterface& View, 
+		const FTransform& ComponentToWorld)> GetRenderVisibilityFunction() const { return RenderVisibilityFunc; }
 	
 	/** Helper method that just sets the same material in all slots. Does not include hover override material. */
 	void SetAllMaterials(UMaterialInterface* Material);
@@ -63,6 +74,9 @@ public:
 	// UPrimitiveComponent
 	virtual bool LineTraceComponent(FHitResult& OutHit, const FVector Start, const FVector End, const FCollisionQueryParams& Params) override;
 
+	// USceneComponent
+	virtual FBoxSphereBounds CalcBounds(const FTransform& LocalToWorld) const override;
+
 	// UActorComponent
 	virtual bool IsHLODRelevant() const override { return false; }
 
@@ -89,4 +103,7 @@ private:
 
 	bool bHovered = false;
 	bool bInteracted = false;
+
+	TFunction<bool(const UE::GizmoRenderingUtil::ISceneViewInterface& View, 
+		const FTransform& ComponentToWorld)> RenderVisibilityFunc = nullptr;
 };
diff --git a/Engine/Source/Runtime/InteractiveToolsFramework/Public/BaseGizmos/ViewBasedTransformAdjusters.h b/Engine/Source/Runtime/InteractiveToolsFramework/Public/BaseGizmos/ViewBasedTransformAdjusters.h
--- a/Engine/Source/Runtime/InteractiveToolsFramework/Public/BaseGizmos/ViewBasedTransformAdjusters.h
+++ b/Engine/Source/Runtime/InteractiveToolsFramework/Public/BaseGizmos/ViewBasedTransformAdjusters.h
@@ -2,7 +2,8 @@
 
 #pragma once
 
-#include "Math/Transform.h" // Compiler will complain with just MathFwd here due to the inline with FTransform return
+#include "Math/BoxSphereBounds.h"
+#include "Math/Transform.h"
 #include "Templates/SharedPointer.h"
 
 class USceneComponent;
@@ -47,6 +48,20 @@ namespace UE::GizmoRenderingUtil
 		 *  caring about the type of adjuster.
 		 */
 		virtual void UpdateWorldLocalState(bool bWorldIn) {}
+
+		/**
+		 * Returns the transform that should be returned by a USceneComponent CalcBounds method, which
+		 *  can't know about the view. Although the default is to return the original bounds, any adjuster
+		 *  that scales the mesh arbitrarily large will probably need to return infinite bounds, whereas
+		 *  an adjuster that only changes orienation must account for any possible orientation.
+		 * Note that returning infinite bounds causes the relevant components to no longer be frustum-culled,
+		 *  but this is acceptible for gizmos, which are typically few and on-screen.
+		 */
+		virtual FBoxSphereBounds GetViewIndependentBounds(const FTransform& LocalToWorld,
+			const FBoxSphereBounds& OriginalBounds)
+		{
+			return OriginalBounds;
+		}
 	};
 
 	/**
@@ -61,6 +76,9 @@ namespace UE::GizmoRenderingUtil
 		INTERACTIVETOOLSFRAMEWORK_API virtual FTransform GetAdjustedComponentToWorld(
 			const ISceneViewInterface& View,
 			const FTransform& CurrentComponentToWorld) override;
+
+		INTERACTIVETOOLSFRAMEWORK_API virtual FBoxSphereBounds GetViewIndependentBounds(const FTransform& LocalToWorld,
+			const FBoxSphereBounds& OriginalBounds) override;
 	private:
 	};
 
@@ -85,6 +103,9 @@ namespace UE::GizmoRenderingUtil
 			//  This is frequently not necessary to set because code upstream will typically update it every tick anyway
 			//  via UViewAdjustedStaticMeshGizmoComponent's UpdateWorldLocalState implementation.
 			bool bUseWorldAxesForGizmo = false;
+			// If positive, keeps the component a given distance in front of the ortho camera, to avoid
+			//  clipping it. Applied before bKeepConstantViewSize.
+			double DistanceInFrontOfCameraInOrtho = 0;
 		};
 
 		FSubGizmoTransformAdjuster(
@@ -100,6 +121,10 @@ namespace UE::GizmoRenderingUtil
 		 */
 		INTERACTIVETOOLSFRAMEWORK_API void SetGizmoOriginTransform(const FTransform& GizmoOriginToComponent);
 
+		void SetMirrorBasedOnOctant(bool bOn) 
+		{
+			Settings.bMirrorBasedOnOctant = bOn;
+		}
 		void SetSettings(const FSettings& SettingsIn)
 		{
 			Settings = SettingsIn;
@@ -116,6 +141,10 @@ namespace UE::GizmoRenderingUtil
 			UViewAdjustedStaticMeshGizmoComponent* ComponentIn, USceneComponent* ComponentToKeepDistanceConstantTo, 
 			bool bMirrorBasedOnOctant);
 
+		INTERACTIVETOOLSFRAMEWORK_API static TSharedPtr<FSubGizmoTransformAdjuster> AddTransformAdjuster(
+			UViewAdjustedStaticMeshGizmoComponent* ComponentIn, USceneComponent* ComponentToKeepDistanceConstantTo,
+			TFunction<bool()> ShouldMirrorBasedOnOctant);
+
 
 		// IGizmoViewTransformAdjuster
 		INTERACTIVETOOLSFRAMEWORK_API virtual FTransform GetAdjustedComponentToWorld(
@@ -129,6 +158,9 @@ namespace UE::GizmoRenderingUtil
 		{
 			Settings.bUseWorldAxesForGizmo = bWorldIn;
 		}
+
+		INTERACTIVETOOLSFRAMEWORK_API virtual FBoxSphereBounds GetViewIndependentBounds(const FTransform& LocalToWorld,
+			const FBoxSphereBounds& OriginalBounds) override;
 	private:
 		FTransform GizmoOriginToComponent_GameThread = FTransform::Identity;
 		FTransform GizmoOriginToComponent_RenderThread = FTransform::Identity;
@@ -160,6 +192,9 @@ namespace UE::GizmoRenderingUtil
 		INTERACTIVETOOLSFRAMEWORK_API virtual FTransform GetAdjustedComponentToWorld(
 			const UE::GizmoRenderingUtil::ISceneViewInterface& View,
 			const FTransform& CurrentComponentToWorld) override;
+
+		INTERACTIVETOOLSFRAMEWORK_API virtual FBoxSphereBounds GetViewIndependentBounds(const FTransform& LocalToWorld,
+			const FBoxSphereBounds& OriginalBounds) override;
 	private:
 		FTransform ViewRelativeTransform = FTransform::Identity;
 		bool bKeepConstantViewSize = true;
