InteractiveToolsFramework: Add free rotation and translation sub gizmos. Factor out UpdateCameraAxisSource so that it can be reused. Also tweak functions such that we can swap the component inside a component HitTarget without having to reset its hover/interacting functions (for cases where the look of the component changes in combined mode).

diff --git a/Engine/Source/Runtime/InteractiveToolsFramework/Private/BaseGizmos/FreePositionSubGizmo.cpp b/Engine/Source/Runtime/InteractiveToolsFramework/Private/BaseGizmos/FreePositionSubGizmo.cpp
new file mode 100644
--- /dev/null
+++ b/Engine/Source/Runtime/InteractiveToolsFramework/Private/BaseGizmos/FreePositionSubGizmo.cpp
@@ -0,0 +1,109 @@
+// Copyright Epic Games, Inc. All Rights Reserved.
+
+#include "BaseGizmos/FreePositionSubGizmo.h"
+
+#include "InteractiveGizmoManager.h"
+#include "BaseBehaviors/ClickDragBehavior.h"
+#include "BaseBehaviors/MouseHoverBehavior.h"
+#include "BaseGizmos/GizmoMath.h"
+#include "BaseGizmos/GizmoPrivateUtil.h" // SetCommonSubGizmoProperties, UpdateCameraAxisSource
+#include "BaseGizmos/GizmoInterfaces.h" // IGizmoTransformSource
+#include "BaseGizmos/TransformSubGizmoUtil.h" // FTransformSubGizmoCommonParams
+#include "Components/PrimitiveComponent.h"
+
+bool UFreePositionSubGizmo::InitializeAsScreenPlaneTranslateGizmo(
+	const UE::GizmoUtil::FTransformSubGizmoCommonParams& Params,
+	UE::GizmoUtil::FTransformSubGizmoSharedState* SharedState)
+{
+	if (!Params.Component
+		|| !Params.TransformProxy)
+	{
+		return false;
+	}
+
+	// The below setup is a merging of InitializeAsUniformScaleGizmo (to set up the camera
+	//  axis source) and InitializeAsTranslateGizmo to set up the translation parameter
+	//  source. But we also store the transform source so that we can use it directly when
+	//  using a custom destination function.
+
+	UGizmoScaledAndUnscaledTransformSources* CastTransformSource;
+
+	// Make sure axis index is invalid so that the SetCommonSubGizmoProperties call below doesn't create
+	// an axis source for us.
+	if (!ensureMsgf(Params.Axis == EAxis::None, TEXT("UFreePositionSubGizmo uses a camera axis source, so axis parameter should be None.")))
+	{
+		UE::GizmoUtil::FTransformSubGizmoCommonParams ParamsCopy = Params;
+		ParamsCopy.Axis = EAxis::None;
+		if (!UE::GizmoUtil::SetCommonSubGizmoProperties(this, ParamsCopy, SharedState, CastTransformSource))
+		{
+			return false;
+		}
+	}
+	else if (!UE::GizmoUtil::SetCommonSubGizmoProperties(this, Params, SharedState, CastTransformSource))
+	{
+		return false;
+	}
+
+	UObject* Owner = Params.OuterForSubobjects ? Params.OuterForSubobjects : GetTransientPackage();
+
+	UGizmoConstantFrameAxisSource* CameraAxisSource = nullptr;
+
+	// See if we already have the axis source in our shared state
+	if (SharedState && SharedState->CameraAxisSource)
+	{
+		CameraAxisSource = SharedState->CameraAxisSource;
+	}
+	else
+	{
+		// Create new and add to shared state.
+		CameraAxisSource = NewObject<UGizmoConstantFrameAxisSource>(Owner);
+		TWeakObjectPtr<AActor> OwnerActor = Params.Component->GetOwner();
+		CustomTickFunction = [this, OwnerActor](float DeltaTime)
+		{
+			if (UGizmoConstantFrameAxisSource* FrameAxisSource = Cast<UGizmoConstantFrameAxisSource>(AxisSource.GetObject()))
+			{
+				UE::GizmoUtil::UpdateCameraAxisSource(
+					*FrameAxisSource,
+					GetGizmoManager(),
+					OwnerActor.IsValid() ? OwnerActor->GetTransform().GetLocation() : FVector3d::ZeroVector);
+			}
+		};
+
+		if (SharedState)
+		{
+			SharedState->CameraAxisSource = CameraAxisSource;
+		}
+	}
+	AxisSource = CameraAxisSource;
+
+	// Parameter source maps axis-parameter-change to translation of TransformSource's transform
+	ParameterSource = UGizmoPlaneTranslationParameterSource::Construct(
+		AxisSource.GetInterface(), CastTransformSource, Owner);
+
+	TransformSource = CastTransformSource;
+
+	return true;
+}
+
+void UFreePositionSubGizmo::OnClickDrag(const FInputDeviceRay& DragPos)
+{
+	if (!bInInteraction)
+	{
+		return;
+	}
+
+	// If we have a custom destination function, we use it to set our transform
+	FCustomDestinationParams Params;
+	Params.WorldRay = &DragPos.WorldRay;
+	FVector HitPoint;
+	if (ShouldUseCustomDestinationFunc() && CustomDestinationFunc(Params, HitPoint))
+	{
+		FTransform CurrentTransform = TransformSource->GetTransform();
+		CurrentTransform.SetLocation(HitPoint);
+		TransformSource->SetTransform(CurrentTransform);
+		return;
+	}
+	// Otherwise, we do regular translation in the plane (we do end up doing a custom
+	//  destination check a second time in there, but it's not worth factoring out).
+	Super::OnClickDrag(DragPos);
+}
diff --git a/Engine/Source/Runtime/InteractiveToolsFramework/Private/BaseGizmos/FreePositionSubGizmo.h b/Engine/Source/Runtime/InteractiveToolsFramework/Private/BaseGizmos/FreePositionSubGizmo.h
new file mode 100644
--- /dev/null
+++ b/Engine/Source/Runtime/InteractiveToolsFramework/Private/BaseGizmos/FreePositionSubGizmo.h
@@ -0,0 +1,37 @@
+// Copyright Epic Games, Inc. All Rights Reserved.
+
+#pragma once
+
+#include "BaseGizmos/PlanePositionGizmo.h"
+
+#include "FreePositionSubGizmo.generated.h"
+
+class IGizmoTransformSource;
+
+/**
+ * UFreePositionSubGizmo is very similar to UPlanePositionGizmo with a camera axis source, 
+ *  but when using a custom destination function, it can use the destination to directly
+ *  set a transform source.
+ */
+UCLASS(MinimalAPI)
+class UFreePositionSubGizmo : public UPlanePositionGizmo
+{
+	GENERATED_BODY()
+
+public:
+	bool InitializeAsScreenPlaneTranslateGizmo(
+		const UE::GizmoUtil::FTransformSubGizmoCommonParams& Params,
+		UE::GizmoUtil::FTransformSubGizmoSharedState* SharedState);
+
+	// IClickDragBehaviorTarget
+	virtual void OnClickDrag(const FInputDeviceRay& DragPos) override;
+
+public:
+	// The below properties can be manipulated for more fine-grained control, but typically it is sufficient
+	// to use one of the initialization methods above.
+
+	/** AxisSource provides the 3D plane on which the interaction happens */
+	UPROPERTY()
+	TScriptInterface<IGizmoTransformSource> TransformSource;
+};
+
diff --git a/Engine/Source/Runtime/InteractiveToolsFramework/Private/BaseGizmos/FreeRotationSubGizmo.cpp b/Engine/Source/Runtime/InteractiveToolsFramework/Private/BaseGizmos/FreeRotationSubGizmo.cpp
new file mode 100644
--- /dev/null
+++ b/Engine/Source/Runtime/InteractiveToolsFramework/Private/BaseGizmos/FreeRotationSubGizmo.cpp
@@ -0,0 +1,314 @@
+// Copyright Epic Games, Inc. All Rights Reserved.
+
+#include "BaseGizmos/FreeRotationSubGizmo.h"
+
+#include "InteractiveGizmoManager.h"
+#include "BaseBehaviors/ClickDragBehavior.h"
+#include "BaseBehaviors/MouseHoverBehavior.h"
+#include "BaseGizmos/AxisSources.h"
+#include "BaseGizmos/GizmoMath.h"
+#include "BaseGizmos/GizmoPrivateUtil.h" // SetCommonSubGizmoProperties
+#include "BaseGizmos/GizmoRenderingUtil.h"
+#include "BaseGizmos/GizmoViewContext.h"
+#include "BaseGizmos/HitTargets.h"
+#include "BaseGizmos/StateTargets.h"
+#include "BaseGizmos/TransformSources.h"
+#include "BaseGizmos/TransformSubGizmoUtil.h" // FTransformSubGizmoCommonParams
+#include "ToolDataVisualizer.h"
+#include "Util/ColorConstants.h" // UE::Geometry::LinearColors
+
+namespace FreeRotationSubGizmoLocals
+{
+	FLinearColor CircleColor = UE::Geometry::LinearColors::Gray3f();
+	int32 CircleNumSections = 32;
+	float CircleThickness = 2.0f;
+}
+
+void UFreeRotationSubGizmo::Setup()
+{
+	UInteractiveGizmo::Setup();
+
+	// Add default mouse input behavior
+	MouseBehavior = NewObject<UClickDragInputBehavior>();
+	MouseBehavior->Initialize(this);
+	MouseBehavior->SetDefaultPriority(FInputCapturePriority(FInputCapturePriority::DEFAULT_GIZMO_PRIORITY));
+	AddInputBehavior(MouseBehavior);
+	UMouseHoverBehavior* HoverBehavior = NewObject<UMouseHoverBehavior>();
+	HoverBehavior->Initialize(this);
+	HoverBehavior->SetDefaultPriority(FInputCapturePriority(FInputCapturePriority::DEFAULT_GIZMO_PRIORITY));
+	AddInputBehavior(HoverBehavior);
+
+	AxisSource = NewObject<UGizmoConstantAxisSource>(this);
+	HitTarget = NewObject<UGizmoComponentHitTarget>(this);
+	StateTarget = NewObject<UGizmoNilStateTarget>(this);
+
+	bInInteraction = false;
+}
+
+
+
+bool UFreeRotationSubGizmo::InitializeAsRotationGizmo(
+	const UE::GizmoUtil::FTransformSubGizmoCommonParams& Params,
+	UGizmoViewContext* GizmoViewContextIn,
+	UE::GizmoUtil::FTransformSubGizmoSharedState* SharedState)
+{
+	if (!Params.Component
+		|| !Params.TransformProxy)
+	{
+		return false;
+	}
+
+	UGizmoScaledAndUnscaledTransformSources* TransformSourcePtr;
+
+	// Make sure axis index is invalid so that the SetCommonSubGizmoProperties call below doesn't create
+	// an axis source for us.
+	if (!ensureMsgf(Params.Axis == EAxis::None, TEXT("UFreeRotationSubGizmo uses a camera axis source, so axis parameter should be None.")))
+	{
+		UE::GizmoUtil::FTransformSubGizmoCommonParams ParamsCopy = Params;
+		ParamsCopy.Axis = EAxis::None;
+		if (!UE::GizmoUtil::SetCommonSubGizmoProperties(this, ParamsCopy, SharedState, TransformSourcePtr))
+		{
+			return false;
+		}
+	}
+	else if (!UE::GizmoUtil::SetCommonSubGizmoProperties(this, Params, SharedState, TransformSourcePtr))
+	{
+		return false;
+	}
+	TransformSource = TransformSourcePtr;
+
+	// Set up the camera axis source
+	UGizmoConstantFrameAxisSource* CameraAxisSource = nullptr;
+
+	// See if we already have it in our shared state
+	if (SharedState && SharedState->CameraAxisSource)
+	{
+		CameraAxisSource = SharedState->CameraAxisSource;
+	}
+	else
+	{
+		// Create new and add to shared state.
+		UObject* Owner = Params.OuterForSubobjects ? Params.OuterForSubobjects : GetTransientPackage();
+		CameraAxisSource = NewObject<UGizmoConstantFrameAxisSource>(Owner);
+		TWeakObjectPtr<AActor> OwnerActor = Params.Component->GetOwner();
+		CustomTickFunction = [this, OwnerActor](float DeltaTime)
+		{
+			if (UGizmoConstantFrameAxisSource* FrameAxisSource = Cast<UGizmoConstantFrameAxisSource>(AxisSource.GetObject()))
+			{
+				UE::GizmoUtil::UpdateCameraAxisSource(
+					*FrameAxisSource,
+					GetGizmoManager(),
+					OwnerActor.IsValid() ? OwnerActor->GetTransform().GetLocation() : FVector3d::ZeroVector);
+			}
+		};
+
+		if (SharedState)
+		{
+			SharedState->CameraAxisSource = CameraAxisSource;
+		}
+	}
+	AxisSource = CameraAxisSource;
+
+	GizmoViewContext = GizmoViewContextIn;
+
+	return true;
+}
+
+// Called from OnClickPress if mode is IncrementalSphereBound
+bool UFreeRotationSubGizmo::ClickPress_IncrementalSphereBound(const FInputDeviceRay& PressPos)
+{
+	if (UnscaledSphereRadius == 0)
+	{
+		return false;
+	}
+
+	FVector SphereCenter = AxisSource->GetOrigin();
+	double LengthScale = UE::GizmoRenderingUtil::CalculateLocalPixelToWorldScale(GizmoViewContext, SphereCenter);
+	InteractionSphereRadius = UnscaledSphereRadius * LengthScale;
+
+	bool bIntersects;
+	GizmoMath::RaySphereIntersection(AxisSource->GetOrigin(),
+		InteractionSphereRadius, PressPos.WorldRay.Origin, PressPos.WorldRay.Direction,
+		bIntersects, LastSphereIntersectionPoint);
+
+	return bIntersects;
+}
+
+// Called from OnClickDrag if mode is IncrementalSphereBound
+bool UFreeRotationSubGizmo::ClickDrag_IncrementalSphereBound(const FInputDeviceRay& DragPos)
+{
+	if (InteractionSphereRadius == 0)
+	{
+		return false;
+	}
+	FVector SphereCenter = AxisSource->GetOrigin();
+
+	bool bIntersects = false;
+	FVector CurrentSphereIntersectionPoint;
+	GizmoMath::RaySphereIntersection(SphereCenter,
+		InteractionSphereRadius, DragPos.WorldRay.Origin, DragPos.WorldRay.Direction,
+		bIntersects, CurrentSphereIntersectionPoint);
+
+	if (!bIntersects)
+	{
+		// If we didn't hit the sphere itself, raycast the plane
+		FVector PlaneHitPoint;
+		GizmoMath::RayPlaneIntersectionPoint(
+			AxisSource->GetOrigin(), AxisSource->GetDirection(),
+			DragPos.WorldRay.Origin, DragPos.WorldRay.Direction,
+			bIntersects, PlaneHitPoint);
+
+		if (!bIntersects)
+		{
+			return false;
+		}
+
+		// Find the point on the sphere that is closest to the hit plane point
+		FVector TowardSphereVector = SphereCenter - PlaneHitPoint;
+		TowardSphereVector.Normalize();
+		GizmoMath::RaySphereIntersection(SphereCenter,
+			InteractionSphereRadius, PlaneHitPoint, TowardSphereVector,
+			bIntersects, CurrentSphereIntersectionPoint);
+	}
+
+	if (!ensure(bIntersects))
+	{
+		return false;
+	}
+
+	// Find the angle that we've rotated the sphere from LastSphereIntersectionPoint to current
+	FVector ToLastSpherePoint = (LastSphereIntersectionPoint - SphereCenter) / InteractionSphereRadius;
+	FVector ToCurrentSpherePoint = (CurrentSphereIntersectionPoint - SphereCenter) / InteractionSphereRadius;
+	FVector RotationAxis = ToLastSpherePoint.Cross(ToCurrentSpherePoint);
+	if (!RotationAxis.Normalize())
+	{
+		return false;
+	}
+	double RotationAngle = acos(ToLastSpherePoint.Dot(ToCurrentSpherePoint));
+	
+	// Apply the new transform to the transform source
+	FTransform CurrentTransform = TransformSource->GetTransform();
+	FQuat RotationToApply(RotationAxis, RotationAngle);
+	CurrentTransform.SetRotation(RotationToApply * CurrentTransform.GetRotation());
+	TransformSource->SetTransform(CurrentTransform);
+	
+	LastSphereIntersectionPoint = CurrentSphereIntersectionPoint;
+
+	return true;
+}
+
+void UFreeRotationSubGizmo::Render(IToolsContextRenderAPI* RenderAPI)
+{
+	using namespace FreeRotationSubGizmoLocals;
+
+	if (bShowSphereBoundsDuringInteraction && bInInteraction)
+	{
+		FVector3d SphereCenter = AxisSource->GetOrigin();
+		double LengthScale = UE::GizmoRenderingUtil::CalculateLocalPixelToWorldScale(RenderAPI->GetSceneView(), SphereCenter);
+		double SphereRadius = UnscaledSphereRadius * LengthScale;
+
+		FToolDataVisualizer Drawer;
+		Drawer.BeginFrame(RenderAPI, RenderAPI->GetCameraState());
+		Drawer.DrawViewFacingCircle(SphereCenter, SphereRadius, CircleNumSections,
+			CircleColor, CircleThickness,/*bDepthTested*/ false);
+		Drawer.EndFrame();
+	}
+}
+
+void UFreeRotationSubGizmo::Tick(float DeltaTime)
+{
+	if (CustomTickFunction)
+	{
+		CustomTickFunction(DeltaTime);
+	}
+}
+
+
+FInputRayHit UFreeRotationSubGizmo::CanBeginClickDragSequence(const FInputDeviceRay& PressPos)
+{
+	FInputRayHit GizmoHit;
+	if (HitTarget && AxisSource && TransformSource)
+	{
+		GizmoHit = HitTarget->IsHit(PressPos);
+	}
+	return GizmoHit;
+}
+
+void UFreeRotationSubGizmo::OnClickPress(const FInputDeviceRay& PressPos)
+{
+	bInInteraction = ClickPress_IncrementalSphereBound(PressPos);
+	if (!bInInteraction)
+	{
+		return;
+	}
+
+	if (StateTarget)
+	{
+		StateTarget->BeginUpdate();
+	}
+	if (ensure(HitTarget))
+	{
+		HitTarget->UpdateInteractingState(bInInteraction);
+	}
+}
+
+void UFreeRotationSubGizmo::OnClickDrag(const FInputDeviceRay& DragPos)
+{
+	if (!bInInteraction)
+	{
+		return;
+	}
+
+	ClickDrag_IncrementalSphereBound(DragPos);
+}
+
+void UFreeRotationSubGizmo::OnClickRelease(const FInputDeviceRay& ReleasePos)
+{
+	OnTerminateDragSequence();
+}
+
+
+void UFreeRotationSubGizmo::OnTerminateDragSequence()
+{
+	if (!bInInteraction)
+	{
+		return;
+	}
+
+	if (StateTarget)
+	{
+		StateTarget->EndUpdate();
+	}
+	bInInteraction = false;
+	if (ensure(HitTarget))
+	{
+		HitTarget->UpdateInteractingState(bInInteraction);
+	}
+}
+
+FInputRayHit UFreeRotationSubGizmo::BeginHoverSequenceHitTest(const FInputDeviceRay& PressPos)
+{
+	FInputRayHit GizmoHit;
+	if (HitTarget)
+	{
+		GizmoHit = HitTarget->IsHit(PressPos);
+	}
+	return GizmoHit;
+}
+
+void UFreeRotationSubGizmo::OnBeginHover(const FInputDeviceRay& DevicePos)
+{
+	HitTarget->UpdateHoverState(true);
+}
+
+bool UFreeRotationSubGizmo::OnUpdateHover(const FInputDeviceRay& DevicePos)
+{
+	// not necessary...
+	HitTarget->UpdateHoverState(true);
+	return true;
+}
+
+void UFreeRotationSubGizmo::OnEndHover()
+{
+	HitTarget->UpdateHoverState(false);
+}
\ No newline at end of file
diff --git a/Engine/Source/Runtime/InteractiveToolsFramework/Private/BaseGizmos/FreeRotationSubGizmo.h b/Engine/Source/Runtime/InteractiveToolsFramework/Private/BaseGizmos/FreeRotationSubGizmo.h
new file mode 100644
--- /dev/null
+++ b/Engine/Source/Runtime/InteractiveToolsFramework/Private/BaseGizmos/FreeRotationSubGizmo.h
@@ -0,0 +1,136 @@
+// Copyright Epic Games, Inc. All Rights Reserved.
+
+#pragma once
+
+#include "BaseBehaviors/BehaviorTargetInterfaces.h"
+#include "InteractiveGizmo.h"
+#include "Math/MathFwd.h"
+
+#include "FreeRotationSubGizmo.generated.h"
+
+class IGizmoAxisSource;
+class IGizmoClickTarget;
+class IGizmoStateTarget;
+class IGizmoTransformSource;
+class UClickDragInputBehavior;
+class UGizmoViewContext;
+namespace UE::GizmoUtil
+{
+	struct FTransformSubGizmoCommonParams;
+	struct FTransformSubGizmoSharedState;
+}
+
+/**
+ * A free rotation sub gizmo implements an arcball-like rotation.
+ */
+UCLASS(MinimalAPI)
+class UFreeRotationSubGizmo : public UInteractiveGizmo, public IClickDragBehaviorTarget, public IHoverBehaviorTarget
+{
+	GENERATED_BODY()
+
+public:
+	/**
+	 * Initializes the properties for the gizmo
+	 */
+	bool InitializeAsRotationGizmo(
+		const UE::GizmoUtil::FTransformSubGizmoCommonParams& InitializationParams,
+		UGizmoViewContext* ViewContext,
+		UE::GizmoUtil::FTransformSubGizmoSharedState* SharedState);
+	
+	// There are many potential approaches to doing a free rotate handle. Currently we
+	//  implement IncrementalSphereBound from the below, which feels pretty intuitive,
+	//  but we could add the other ones if desired.
+	//enum class EMode
+	//{
+	//	// Intersect rays with a sphere and rotate the latest result such that the previous 
+	//	//  intersection goes to the latest intersection. Outside of sphere, assume you are
+	//	//  on the tangent side of the sphere.
+	//	// This is similar to how the normal editor gizmo works if you enable arcball rotation.
+	//	IncrementalSphereBound,
+	//	
+	//	// Like IncrementalSphereBound, but rotate the original (not latest) transform such
+	//	//  that the first (i.e. original) intersection goes to latest intersection.
+	//	// This is similar to how the ball works in the new gizmos (i.e. if Enable New Gizmos) is true.
+	//	SourceToDestSphereBound,
+
+	//	// Similar to SourceToDestSphereBound in that we're determining our rotation axis based on
+	//	//  the first intersection going to latest intersection, but the amount to rotate is determined
+	//	//  by the distance in the camera plane, so that we can keep rotating the object 360 degrees or
+	//	//  more with one drag.
+	//	SourceToDestUnbounded
+	//};
+	//void SetMode(EMode::ModeIn);
+	
+	/**
+	 * Determines the size of the invisible sphere we raycast to perform the rotation.
+	 */ 
+	void SetUnscaledSphereRadius(double Radius) { UnscaledSphereRadius = Radius; }
+
+	/**
+	 * When true (default) a circle is drawn to show the outside bounds of the interaction sphere while
+	 *  interacting with the gizmo.
+	 */
+	void SetShowSphereBoundsDuringInteraction(bool bOn) { bShowSphereBoundsDuringInteraction = bOn; }
+
+	// UInteractiveGizmo overrides
+	virtual void Setup() override;
+	virtual void Render(IToolsContextRenderAPI* RenderAPI) override;
+	virtual void Tick(float DeltaTime) override;
+
+	// IClickDragBehaviorTarget implementation
+	virtual FInputRayHit CanBeginClickDragSequence(const FInputDeviceRay& PressPos) override;
+	virtual void OnClickPress(const FInputDeviceRay& PressPos) override;
+	virtual void OnClickDrag(const FInputDeviceRay& DragPos) override;
+	virtual void OnClickRelease(const FInputDeviceRay& ReleasePos) override;
+	virtual void OnTerminateDragSequence() override;
+
+	// IHoverBehaviorTarget implementation
+	virtual FInputRayHit BeginHoverSequenceHitTest(const FInputDeviceRay& PressPos) override;
+	virtual void OnBeginHover(const FInputDeviceRay& DevicePos) override;
+	virtual bool OnUpdateHover(const FInputDeviceRay& DevicePos) override;
+	virtual void OnEndHover() override;
+
+public:
+	// The below properties can be manipulated for more fine-grained control, but typically it is sufficient
+	// to use the initialization method above.
+
+	/** AxisSource provides the origin of the interaction sphere and the plane to raycast when hitting outside the sphere */
+	UPROPERTY()
+	TScriptInterface<IGizmoAxisSource> AxisSource;
+	
+	/** The HitTarget provides a hit-test against some 3D element (presumably a visual widget) that controls when interaction can start */
+	UPROPERTY()
+	TScriptInterface<IGizmoClickTarget> HitTarget;
+
+	/** StateTarget is notified when interaction starts and ends, so that things like undo/redo can be handled externally */
+	UPROPERTY()
+	TScriptInterface<IGizmoStateTarget> StateTarget;
+
+	/** Target that is rotated by the sub gizmo. */
+	UPROPERTY()
+	TScriptInterface<IGizmoTransformSource> TransformSource;
+
+	/** View info used during raycasts */
+	UPROPERTY()
+	TObjectPtr<UGizmoViewContext> GizmoViewContext = nullptr;
+
+	/** The mouse click behavior of the gizmo is accessible so that it can be modified to use different mouse keys. */
+	UPROPERTY()
+	TObjectPtr<UClickDragInputBehavior> MouseBehavior;
+	
+private:
+	bool bInInteraction = false;
+	bool bShowSphereBoundsDuringInteraction = true;
+
+	FVector LastSphereIntersectionPoint;
+	double InteractionSphereRadius = 0;
+	bool ClickPress_IncrementalSphereBound(const FInputDeviceRay& PressPos);
+	bool ClickDrag_IncrementalSphereBound(const FInputDeviceRay& DragPos);
+	
+	double UnscaledSphereRadius = 100;
+
+	// Helper that can hold some extra upkeep to do during Tick (used to
+	//  update a camera axis source if needed).
+	TUniqueFunction<void(float DeltaTime)> CustomTickFunction;
+};
+
diff --git a/Engine/Source/Runtime/InteractiveToolsFramework/Private/BaseGizmos/GizmoPrivateUtil.cpp b/Engine/Source/Runtime/InteractiveToolsFramework/Private/BaseGizmos/GizmoPrivateUtil.cpp
--- a/Engine/Source/Runtime/InteractiveToolsFramework/Private/BaseGizmos/GizmoPrivateUtil.cpp
+++ b/Engine/Source/Runtime/InteractiveToolsFramework/Private/BaseGizmos/GizmoPrivateUtil.cpp
@@ -6,6 +6,8 @@
 #include "BaseGizmos/AxisAngleGizmo.h"
 #include "BaseGizmos/AxisPositionGizmo.h"
 #include "BaseGizmos/CombinedTransformGizmo.h"
+#include "BaseGizmos/FreePositionSubGizmo.h"
+#include "BaseGizmos/FreeRotationSubGizmo.h"
 #include "BaseGizmos/GizmoBaseComponent.h"
 #include "BaseGizmos/GizmoViewContext.h"
 #include "BaseGizmos/HitTargets.h"
@@ -67,6 +69,24 @@ UGizmoViewContext* UE::GizmoUtil::GetGizmoViewContext(UInteractiveGizmoManager*
 	return GizmoViewContext;
 }
 
+bool UE::GizmoUtil::UpdateCameraAxisSource(UGizmoConstantFrameAxisSource& CameraAxisSource, 
+	UInteractiveGizmoManager* GizmoManager, const FVector3d& AxisOrigin)
+{
+	if (!GizmoManager || !GizmoManager->GetContextQueriesAPI())
+	{
+		return false;
+	}
+
+	FViewCameraState CameraState;
+	GizmoManager->GetContextQueriesAPI()->GetCurrentViewState(CameraState);
+
+	CameraAxisSource.Origin = AxisOrigin;
+	CameraAxisSource.Direction = -CameraState.Forward();
+	CameraAxisSource.TangentX = CameraState.Right();
+	CameraAxisSource.TangentY = CameraState.Up();
+	return true;
+}
+
 template <typename SubGizmoType>
 bool UE::GizmoUtil::SetCommonSubGizmoProperties(
 	SubGizmoType* Gizmo,
@@ -151,16 +171,18 @@ bool UE::GizmoUtil::SetCommonSubGizmoProperties(
 		// TODO: See comment in FTransformSubGizmoCommonParams- seems not worth giving the user the ability to
 		//  customize this function.
 		//Params.CustomUpdateHoverFunction ? Params.CustomUpdateHoverFunction : 
-		[Component = Params.Component](bool bHovering)
+		[CastHitTarget](bool bHovering)
 	{
-		if (IGizmoBaseComponentInterface* Hoverable = Cast<IGizmoBaseComponentInterface>(Component))
-		{
+		if (!ensure(IsValid(CastHitTarget))) return;
+		if (IGizmoBaseComponentInterface* Hoverable = Cast<IGizmoBaseComponentInterface>(CastHitTarget->Component))
+ 		{
 			Hoverable->UpdateHoverState(bHovering);
 		}
 	};
-	CastHitTarget->UpdateInteractingFunction = [Component = Params.Component](bool bInteracting)
+	CastHitTarget->UpdateInteractingFunction = [CastHitTarget](bool bInteracting)
 	{
-		if (IGizmoBaseComponentInterface* Hoverable = Cast<IGizmoBaseComponentInterface>(Component))
+		if (!ensure(IsValid(CastHitTarget))) return;
+		if (IGizmoBaseComponentInterface* Hoverable = Cast<IGizmoBaseComponentInterface>(CastHitTarget->Component))
 		{
 			Hoverable->UpdateInteractingState(bInteracting);
 		}
@@ -212,4 +234,14 @@ template bool UE::GizmoUtil::SetCommonSubGizmoProperties<UPlanePositionGizmo>(
 	UPlanePositionGizmo* Gizmo,
 	const FTransformSubGizmoCommonParams& Params,
 	FTransformSubGizmoSharedState* SharedState,
+	UGizmoScaledAndUnscaledTransformSources*& TransformSourceOut);
+template bool UE::GizmoUtil::SetCommonSubGizmoProperties<UFreePositionSubGizmo>(
+	UFreePositionSubGizmo* Gizmo,
+	const FTransformSubGizmoCommonParams& Params,
+	FTransformSubGizmoSharedState* SharedState,
+	UGizmoScaledAndUnscaledTransformSources*& TransformSourceOut);
+template bool UE::GizmoUtil::SetCommonSubGizmoProperties<UFreeRotationSubGizmo>(
+	UFreeRotationSubGizmo* Gizmo,
+	const FTransformSubGizmoCommonParams& Params,
+	FTransformSubGizmoSharedState* SharedState,
 	UGizmoScaledAndUnscaledTransformSources*& TransformSourceOut);
\ No newline at end of file
diff --git a/Engine/Source/Runtime/InteractiveToolsFramework/Private/BaseGizmos/GizmoPrivateUtil.h b/Engine/Source/Runtime/InteractiveToolsFramework/Private/BaseGizmos/GizmoPrivateUtil.h
--- a/Engine/Source/Runtime/InteractiveToolsFramework/Private/BaseGizmos/GizmoPrivateUtil.h
+++ b/Engine/Source/Runtime/InteractiveToolsFramework/Private/BaseGizmos/GizmoPrivateUtil.h
@@ -4,11 +4,12 @@
 
 #include "InteractiveGizmo.h" // ETransformGizmoSubElements
 #include "Math/Axis.h"
+#include "VectorTypes.h"
 
+class UGizmoConstantFrameAxisSource;
 class UGizmoScaledAndUnscaledTransformSources;
 class UGizmoViewContext;
 class UInteractiveGizmoManager;
-
 namespace UE::GizmoUtil
 {
 	struct FTransformSubGizmoCommonParams;
@@ -34,6 +35,12 @@ namespace UE::GizmoUtil
 	 */
 	UGizmoViewContext* GetGizmoViewContext(UInteractiveGizmoManager* GizmoManager);
 
+	/**
+	 * Helper to update an axis source that is based on the camera location, called from tick.
+	 */
+	bool UpdateCameraAxisSource(UGizmoConstantFrameAxisSource& CameraAxisSource,
+		UInteractiveGizmoManager* GizmoManager, const FVector3d& AxisOrigin);
+
 	/**
 	 * Helper that sets some common properties on sub gizmos, namely AxisSource, HitTarget, 
 	 *  and StateTarget. Also returns a TransformSource, which typically gets wrapped up in
@@ -48,5 +55,4 @@ namespace UE::GizmoUtil
 		const UE::GizmoUtil::FTransformSubGizmoCommonParams& Params,
 		UE::GizmoUtil::FTransformSubGizmoSharedState* SharedState,
 		UGizmoScaledAndUnscaledTransformSources*& TransformSourceOut);
-
 }// end UE::GizmoUtil
\ No newline at end of file
diff --git a/Engine/Source/Runtime/InteractiveToolsFramework/Private/BaseGizmos/PlanePositionGizmo.cpp b/Engine/Source/Runtime/InteractiveToolsFramework/Private/BaseGizmos/PlanePositionGizmo.cpp
--- a/Engine/Source/Runtime/InteractiveToolsFramework/Private/BaseGizmos/PlanePositionGizmo.cpp
+++ b/Engine/Source/Runtime/InteractiveToolsFramework/Private/BaseGizmos/PlanePositionGizmo.cpp
@@ -5,7 +5,7 @@
 #include "BaseBehaviors/ClickDragBehavior.h"
 #include "BaseBehaviors/MouseHoverBehavior.h"
 #include "BaseGizmos/GizmoMath.h"
-#include "BaseGizmos/GizmoPrivateUtil.h" // SetCommonSubGizmoProperties
+#include "BaseGizmos/GizmoPrivateUtil.h" // SetCommonSubGizmoProperties, UpdateCameraSource
 #include "BaseGizmos/TransformSubGizmoUtil.h" // FTransformSubGizmoCommonParams
 #include "Components/PrimitiveComponent.h"
 
@@ -13,22 +13,7 @@
 
 namespace PlanePositionGizmoLocals
 {
-	bool UpdateCameraAxisSource(UGizmoConstantFrameAxisSource& CameraAxisSource, UInteractiveGizmoManager* GizmoManager, const FVector3d& AxisOrigin)
-	{
-		if (!GizmoManager || !GizmoManager->GetContextQueriesAPI())
-		{
-			return false;
-		}
-
-		FViewCameraState CameraState;
-		GizmoManager->GetContextQueriesAPI()->GetCurrentViewState(CameraState);
 
-		CameraAxisSource.Origin = AxisOrigin;
-		CameraAxisSource.Direction = -CameraState.Forward();
-		CameraAxisSource.TangentX = CameraState.Right();
-		CameraAxisSource.TangentY = CameraState.Up();
-		return true;
-	}
 }
 
 
@@ -195,7 +180,7 @@ bool UPlanePositionGizmo::InitializeAsUniformScaleGizmo(
 		{
 			if (UGizmoConstantFrameAxisSource* FrameAxisSource = Cast<UGizmoConstantFrameAxisSource>(AxisSource.GetObject()))
 			{
-				UpdateCameraAxisSource(
+				UE::GizmoUtil::UpdateCameraAxisSource(
 					*FrameAxisSource,
 					GetGizmoManager(),
 					OwnerActor.IsValid() ? OwnerActor->GetTransform().GetLocation() : FVector3d::ZeroVector);
diff --git a/Engine/Source/Runtime/InteractiveToolsFramework/Public/BaseGizmos/PlanePositionGizmo.h b/Engine/Source/Runtime/InteractiveToolsFramework/Public/BaseGizmos/PlanePositionGizmo.h
--- a/Engine/Source/Runtime/InteractiveToolsFramework/Public/BaseGizmos/PlanePositionGizmo.h
+++ b/Engine/Source/Runtime/InteractiveToolsFramework/Public/BaseGizmos/PlanePositionGizmo.h
@@ -222,7 +222,8 @@ protected:
 	FVector2D InteractionStartOriginParameterOffset;
 	FInputDeviceRay LastInputRay = FInputDeviceRay(FRay());
 
-private:
+	// Helper that can hold some extra upkeep to do during Tick (used to
+	//  update a camera axis source if needed).
 	TUniqueFunction<void(float DeltaTime)> CustomTickFunction;
 };
 
