InteractiveToolsFramework: Use the new meshes in the modeling tools combined gizmo. Add the new free translate and rotate gizmos. Add ability to swap out the look of the axis scale gizmos when not combined.

diff --git a/Engine/Plugins/Experimental/MeshModelingToolsetExp/Source/ModelingEditorUI/Private/STransformGizmoNumericalUIOverlay.cpp b/Engine/Plugins/Experimental/MeshModelingToolsetExp/Source/ModelingEditorUI/Private/STransformGizmoNumericalUIOverlay.cpp
--- a/Engine/Plugins/Experimental/MeshModelingToolsetExp/Source/ModelingEditorUI/Private/STransformGizmoNumericalUIOverlay.cpp
+++ b/Engine/Plugins/Experimental/MeshModelingToolsetExp/Source/ModelingEditorUI/Private/STransformGizmoNumericalUIOverlay.cpp
@@ -315,7 +315,8 @@ TSharedRef<SWidget> STransformGizmoNumericalUIOverlay::CreateWidgetForGizmo(ETra
 
 	// Create the translation row
 	if ((GizmoElements & ETransformGizmoSubElements::TranslateAllAxes) != ETransformGizmoSubElements::None
-		|| (GizmoElements & ETransformGizmoSubElements::TranslateAllPlanes) != ETransformGizmoSubElements::None)
+		|| (GizmoElements & ETransformGizmoSubElements::TranslateAllPlanes) != ETransformGizmoSubElements::None
+		|| (GizmoElements & ETransformGizmoSubElements::FreeTranslate) != ETransformGizmoSubElements::None)
 	{
 		TSharedRef<SHorizontalBox> Row = SNew(SHorizontalBox);
 		Rows->AddSlot()
@@ -336,13 +337,16 @@ TSharedRef<SWidget> STransformGizmoNumericalUIOverlay::CreateWidgetForGizmo(ETra
 
 		bool bHaveXComponent = (GizmoElements & ETransformGizmoSubElements::TranslateAxisX) != ETransformGizmoSubElements::None
 			|| (GizmoElements & ETransformGizmoSubElements::TranslatePlaneXY) != ETransformGizmoSubElements::None
-			|| (GizmoElements & ETransformGizmoSubElements::TranslatePlaneXZ) != ETransformGizmoSubElements::None;
+			|| (GizmoElements & ETransformGizmoSubElements::TranslatePlaneXZ) != ETransformGizmoSubElements::None
+			|| (GizmoElements & ETransformGizmoSubElements::FreeTranslate) != ETransformGizmoSubElements::None;
 		bool bHaveYComponent = (GizmoElements & ETransformGizmoSubElements::TranslateAxisY) != ETransformGizmoSubElements::None
 			|| (GizmoElements & ETransformGizmoSubElements::TranslatePlaneXY) != ETransformGizmoSubElements::None
-			|| (GizmoElements & ETransformGizmoSubElements::TranslatePlaneYZ) != ETransformGizmoSubElements::None;
+			|| (GizmoElements & ETransformGizmoSubElements::TranslatePlaneYZ) != ETransformGizmoSubElements::None
+			|| (GizmoElements & ETransformGizmoSubElements::FreeTranslate) != ETransformGizmoSubElements::None;
 		bool bHaveZComponent = (GizmoElements & ETransformGizmoSubElements::TranslateAxisZ) != ETransformGizmoSubElements::None
 			|| (GizmoElements & ETransformGizmoSubElements::TranslatePlaneXZ) != ETransformGizmoSubElements::None
-			|| (GizmoElements & ETransformGizmoSubElements::TranslatePlaneYZ) != ETransformGizmoSubElements::None;
+			|| (GizmoElements & ETransformGizmoSubElements::TranslatePlaneYZ) != ETransformGizmoSubElements::None
+			|| (GizmoElements & ETransformGizmoSubElements::FreeTranslate) != ETransformGizmoSubElements::None;
 
 		Row->AddSlot()
 		[
@@ -359,7 +363,8 @@ TSharedRef<SWidget> STransformGizmoNumericalUIOverlay::CreateWidgetForGizmo(ETra
 	}
 
 	// Add rotate row
-	if ((GizmoElements & ETransformGizmoSubElements::RotateAllAxes) != ETransformGizmoSubElements::None)
+	if ((GizmoElements & ETransformGizmoSubElements::RotateAllAxes) != ETransformGizmoSubElements::None
+		|| (GizmoElements & ETransformGizmoSubElements::FreeRotate) != ETransformGizmoSubElements::None)
 	{
 		TSharedRef<SHorizontalBox> Row = SNew(SHorizontalBox);
 		Rows->AddSlot()
@@ -378,9 +383,12 @@ TSharedRef<SWidget> STransformGizmoNumericalUIOverlay::CreateWidgetForGizmo(ETra
 			.ColorAndOpacity(FSlateColor::UseForeground())
 		];
 
-		bool bHaveXComponent = (GizmoElements & ETransformGizmoSubElements::RotateAxisX) != ETransformGizmoSubElements::None;
-		bool bHaveYComponent = (GizmoElements & ETransformGizmoSubElements::RotateAxisY) != ETransformGizmoSubElements::None;
-		bool bHaveZComponent = (GizmoElements & ETransformGizmoSubElements::RotateAxisZ) != ETransformGizmoSubElements::None;
+		bool bHaveXComponent = (GizmoElements & ETransformGizmoSubElements::RotateAxisX) != ETransformGizmoSubElements::None
+			|| (GizmoElements & ETransformGizmoSubElements::FreeRotate) != ETransformGizmoSubElements::None;
+		bool bHaveYComponent = (GizmoElements & ETransformGizmoSubElements::RotateAxisY) != ETransformGizmoSubElements::None
+			|| (GizmoElements & ETransformGizmoSubElements::FreeRotate) != ETransformGizmoSubElements::None;
+		bool bHaveZComponent = (GizmoElements & ETransformGizmoSubElements::RotateAxisZ) != ETransformGizmoSubElements::None
+			|| (GizmoElements & ETransformGizmoSubElements::FreeRotate) != ETransformGizmoSubElements::None;
 
 		Row->AddSlot()
 		[
diff --git a/Engine/Source/Runtime/InteractiveToolsFramework/Private/BaseGizmos/CombinedTransformGizmo.cpp b/Engine/Source/Runtime/InteractiveToolsFramework/Private/BaseGizmos/CombinedTransformGizmo.cpp
--- a/Engine/Source/Runtime/InteractiveToolsFramework/Private/BaseGizmos/CombinedTransformGizmo.cpp
+++ b/Engine/Source/Runtime/InteractiveToolsFramework/Private/BaseGizmos/CombinedTransformGizmo.cpp
@@ -4,9 +4,12 @@
 #include "InteractiveGizmoManager.h"
 #include "SceneQueries/SceneSnappingManager.h"
 #include "BaseGizmos/AxisPositionGizmo.h"
+#include "BaseGizmos/FreePositionSubGizmo.h"
+#include "BaseGizmos/FreeRotationSubGizmo.h"
 #include "BaseGizmos/PlanePositionGizmo.h"
 #include "BaseGizmos/AxisAngleGizmo.h"
 #include "BaseGizmos/GizmoComponents.h"
+#include "BaseGizmos/GizmoUtil.h"
 
 #include "BaseGizmos/GizmoArrowComponent.h"
 #include "BaseGizmos/GizmoRectangleComponent.h"
@@ -14,8 +17,11 @@
 #include "BaseGizmos/GizmoBoxComponent.h"
 #include "BaseGizmos/GizmoLineHandleComponent.h"
 #include "BaseGizmos/GizmoPrivateUtil.h" // ToAxis
+#include "BaseGizmos/GizmoRenderingUtil.h"
 #include "BaseGizmos/GizmoViewContext.h"
 #include "BaseGizmos/TransformSubGizmoUtil.h" // FTransformSubGizmoCommonParams, FTransformSubGizmoSharedState
+#include "BaseGizmos/ViewAdjustedStaticMeshGizmoComponent.h"
+#include "BaseGizmos/ViewBasedTransformAdjusters.h" // FSubGizmoTransformAdjuster
 
 #include "Quaternion.h"
 #include "MathUtil.h"
@@ -30,6 +36,8 @@
 #include "ContextObjectStore.h"
 #include "Engine/World.h"
 #include "Engine/CollisionProfile.h"
+#include "Engine/StaticMesh.h"
+#include "HAL/IConsoleManager.h" // FAutoConsoleVariableRef
 
 #include UE_INLINE_GENERATED_CPP_BY_NAME(CombinedTransformGizmo)
 
@@ -38,6 +46,107 @@
 
 namespace CombinedTransformGizmoLocals
 {
+	const int32 DrawModeValue_Meshes = 1;
+	
+	// CVar that determines how we draw the gizmo
+	int32 GizmoDrawMode = DrawModeValue_Meshes;
+	static FAutoConsoleVariableRef CVarGizmoDrawMode(
+		TEXT("modeling.Gizmo.DrawMode"),
+		GizmoDrawMode,
+		TEXT("When 0, modeling gizmos are drawn using the old PDI system. When 1, modeling gizmos use new adjusted-size components. "
+			"Gizmos have to be recreated (by restarting mode/tools) for the change to take effect."));
+
+	// Helper functions that determine whether parts of the gizmo should be visible when using DrawModeValue_Meshes. For example
+	//  we don't want the size of a rotation component to be hiding the axis behind it in ortho view.
+	auto ShouldAxisBeVisible = [](const UE::GizmoRenderingUtil::ISceneViewInterface& View,
+		const FTransform& ComponentToWorld) 
+		{
+			static const double ARROW_RENDERVISIBILITY_DOT_THRESHOLD = FMath::Cos(FMath::DegreesToRadians(3));
+
+			FVector ViewDirection = View.IsPerspectiveProjection() ?
+				ComponentToWorld.GetLocation() - View.GetViewLocation() : View.GetViewDirection();
+			ViewDirection.Normalize();
+			FVector ArrowDirection = ComponentToWorld.TransformVector(FVector::XAxisVector);
+			ArrowDirection.Normalize();
+
+			return FMath::Abs(FVector::DotProduct(ArrowDirection, ViewDirection)) <= ARROW_RENDERVISIBILITY_DOT_THRESHOLD;
+		};
+	auto ShouldPlaneBeVisible = [](const UE::GizmoRenderingUtil::ISceneViewInterface& View,
+		const FTransform& ComponentToWorld) 
+		{
+			static const double RECTANGLE_RENDERVISIBILITY_DOT_THRESHOLD = FMath::Cos(FMath::DegreesToRadians(87));;
+
+			FVector ViewDirection = View.IsPerspectiveProjection() ?
+				ComponentToWorld.GetLocation() - View.GetViewLocation() : View.GetViewDirection();
+			ViewDirection.Normalize();
+			FVector PlaneNormal = ComponentToWorld.TransformVector(FVector::XAxisVector);
+			PlaneNormal.Normalize();
+
+			return
+				FMath::Abs(FVector::DotProduct(PlaneNormal, ViewDirection)) >= RECTANGLE_RENDERVISIBILITY_DOT_THRESHOLD;
+		};
+	
+	FLinearColor FreeRotateColor(0.5f, 0.5f, 0.5f, 0.15f); // A faint translucent gray
+	// Slightly gray so that the selection highlight pops a bit more
+	FLinearColor FreeTranslateColor(0.7f, 0.7f, 0.7f, 1.0f);
+	FLinearColor UniformScaleColor = FreeTranslateColor;
+	FVector CornerScalePositionCombined(0, 120, 120);
+	FVector CornerScalePositionSeparate(0, 75, 75);
+	double CornerScaleHandleScale = 0.5;
+
+	// Helper functions that get the appropriate values to use from an EAxis value, so
+	//  that we can write helpers that just take that as an argument.
+	FLinearColor AxisToLegacyColor(EAxis::Type Axis)
+	{
+		switch (Axis)
+		{
+		case EAxis::X:
+			return FLinearColor::Red;
+		case EAxis::Y:
+			return FLinearColor::Green;
+		case EAxis::Z:
+			return FLinearColor::Blue;
+		default:
+			ensure(false);
+		}
+		return FLinearColor::Black;
+	}
+	FVector AxisToVector(EAxis::Type Axis)
+	{
+		switch (Axis)
+		{
+		case EAxis::X:
+			return FVector::XAxisVector;
+		case EAxis::Y:
+			return FVector::YAxisVector;
+		case EAxis::Z:
+			return FVector::ZAxisVector;
+		default:
+			ensure(false);
+		}
+		return FVector::XAxisVector;
+	}
+	void AxisToLegacyPairOfVectors(EAxis::Type Axis, FVector& Vector1, FVector& Vector2)
+	{
+		switch (Axis)
+		{
+		case EAxis::X:
+			Vector1 = FVector::YAxisVector;
+			Vector2 = FVector::ZAxisVector;
+			return;
+		case EAxis::Y:
+			Vector1 = FVector::XAxisVector;
+			Vector2 = FVector::ZAxisVector;
+			return;
+		case EAxis::Z:
+			Vector1 = FVector::XAxisVector;
+			Vector2 = FVector::YAxisVector;
+			return;
+		default:
+			ensure(false);
+		}
+	}
+
 	// Looks at a gizmo actor and figures out what sub element flags must have been active when creating it.
 	ETransformGizmoSubElements GetSubElementFlagsFromActor(ACombinedTransformGizmoActor* GizmoActor)
 	{
@@ -53,10 +162,12 @@ namespace CombinedTransformGizmoLocals
 		if (GizmoActor->TranslateXY) { Elements |= ETransformGizmoSubElements::TranslatePlaneXY; }
 		if (GizmoActor->TranslateYZ) { Elements |= ETransformGizmoSubElements::TranslatePlaneYZ; }
 		if (GizmoActor->TranslateXZ) { Elements |= ETransformGizmoSubElements::TranslatePlaneXZ; }
+		if (GizmoActor->FreeTranslateHandle) { Elements |= ETransformGizmoSubElements::FreeTranslate; }
 
 		if (GizmoActor->RotateX) { Elements |= ETransformGizmoSubElements::RotateAxisX; }
 		if (GizmoActor->RotateY) { Elements |= ETransformGizmoSubElements::RotateAxisY; }
 		if (GizmoActor->RotateZ) { Elements |= ETransformGizmoSubElements::RotateAxisZ; }
+		if (GizmoActor->FreeRotateHandle) { Elements |= ETransformGizmoSubElements::FreeRotate; }
 
 		if (GizmoActor->AxisScaleX) { Elements |= ETransformGizmoSubElements::ScaleAxisX; }
 		if (GizmoActor->AxisScaleY) { Elements |= ETransformGizmoSubElements::ScaleAxisY; }
@@ -151,12 +262,27 @@ bool ACombinedTransformGizmoActor::ReplaceSubGizmoComponent(ETransformGizmoSubEl
 		break;
 	case ETransformGizmoSubElements::ScaleAxisX:
 		ReplaceComponent(AxisScaleX);
+		if (FullAxisScaleX)
+		{
+			FullAxisScaleX->DestroyComponent();
+			FullAxisScaleX = nullptr;
+		}
 		break;
 	case ETransformGizmoSubElements::ScaleAxisY:
 		ReplaceComponent(AxisScaleY);
+		if (FullAxisScaleY)
+		{
+			FullAxisScaleY->DestroyComponent();
+			FullAxisScaleY = nullptr;
+		}
 		break;
 	case ETransformGizmoSubElements::ScaleAxisZ:
 		ReplaceComponent(AxisScaleZ);
+		if (FullAxisScaleZ)
+		{
+			FullAxisScaleZ->DestroyComponent();
+			FullAxisScaleZ = nullptr;
+		}
 		break;
 	case ETransformGizmoSubElements::ScalePlaneXY:
 		ReplaceComponent(PlaneScaleXY);
@@ -174,6 +300,12 @@ bool ACombinedTransformGizmoActor::ReplaceSubGizmoComponent(ETransformGizmoSubEl
 	case ETransformGizmoSubElements::RotateAllAxes:
 		ReplaceComponent(RotationSphere);
 		break;
+	case ETransformGizmoSubElements::FreeRotate:
+		ReplaceComponent(FreeRotateHandle);
+		break;
+	case ETransformGizmoSubElements::FreeTranslate:
+		ReplaceComponent(FreeTranslateHandle);
+		break;
 	default:
 		UE_LOG(LogGeometry, Warning, TEXT("UCombinedTransformGizmo::SetSubGizmoComponent currently only supports a "
 			"single sub gizmo element at a time."));
@@ -199,15 +331,70 @@ ACombinedTransformGizmoActor* ACombinedTransformGizmoActor::ConstructCustom3Axis
 	UWorld* World, UGizmoViewContext* GizmoViewContext,
 	ETransformGizmoSubElements Elements)
 {
+	using namespace CombinedTransformGizmoLocals;
+	using FSubGizmoTransformAdjuster = UE::GizmoRenderingUtil::FSubGizmoTransformAdjuster;
+
 	FActorSpawnParameters SpawnInfo;
 	ACombinedTransformGizmoActor* NewActor = World->SpawnActor<ACombinedTransformGizmoActor>(FVector::ZeroVector, FRotator::ZeroRotator, SpawnInfo);
 
 	float GizmoLineThickness = 3.0f;
 
+	enum class EMirror 
+	{
+		Always,
+		WhenCombined,
+		Never
+	};
 
-	auto MakeAxisArrowFunc = [&](const FLinearColor& Color, const FVector& Axis)
+	// Helper for adding a mesh-based sub gizmo component (when using DrawModeValue_Meshes)
+	auto AddMeshGizmoComponent = [GizmoViewContext, NewActor](const TCHAR* MeshPath, const FLinearColor& Color, 
+		const FTransform& RelativeTransform, EMirror Mirror, bool bAddHoverMaterial = true) -> UViewAdjustedStaticMeshGizmoComponent*
 	{
-		UGizmoArrowComponent* Component = AddDefaultArrowComponent(World, NewActor, GizmoViewContext, Color, Axis, 60.0f);
+		UStaticMesh* Mesh = LoadObject<UStaticMesh>(nullptr, MeshPath);
+		if (!ensure(Mesh))
+		{
+			return nullptr;
+		}
+
+		UViewAdjustedStaticMeshGizmoComponent* Component = UE::GizmoRenderingUtil::CreateDefaultMaterialGizmoMeshComponent(
+			Mesh, GizmoViewContext, NewActor, Color, bAddHoverMaterial);
+		if (!ensure(Component))
+		{
+			return nullptr;
+		}
+		NewActor->AddInstanceComponent(Component);
+		Component->AttachToComponent(NewActor->GetRootComponent(), FAttachmentTransformRules::KeepRelativeTransform);
+		Component->SetRelativeTransform(RelativeTransform);
+		Component->RegisterComponent();
+		TSharedPtr<FSubGizmoTransformAdjuster> Adjuster = FSubGizmoTransformAdjuster::AddTransformAdjuster(
+				Component, NewActor->GetRootComponent(), Mirror == EMirror::Always);
+		if (Mirror == EMirror::WhenCombined)
+		{
+			NewActor->AdjustersThatMirrorOnlyInCombinedMode.Add(Adjuster);
+		}
+
+		return Component;
+	};
+
+	auto MakeAxisArrowFunc = [World, NewActor, GizmoViewContext, GizmoLineThickness, &AddMeshGizmoComponent](EAxis::Type ElementAxis) -> UPrimitiveComponent*
+	{
+		if (GizmoDrawMode == DrawModeValue_Meshes)
+		{
+			if (UViewAdjustedStaticMeshGizmoComponent* Component = AddMeshGizmoComponent(
+				TEXT("/Engine/InteractiveToolsFramework/Meshes/GizmoArrowHandle"),
+				UE::GizmoRenderingUtil::GetDefaultAxisColor(ElementAxis),
+				UE::GizmoUtil::GetRotatedBasisTransform(
+					// Transform for the X axis, relative to gizmo root
+					FTransform(FQuat::Identity, FVector::ZeroVector, FVector::OneVector), ElementAxis),
+				EMirror::WhenCombined))
+			{
+				Component->SetRenderVisibilityFunction(ShouldAxisBeVisible);
+				return Component;
+			}
+		}
+
+		UGizmoArrowComponent* Component = AddDefaultArrowComponent(World, NewActor, GizmoViewContext, 
+			AxisToLegacyColor(ElementAxis), AxisToVector(ElementAxis), 60.0f);
 		Component->Gap = 20.0f;
 		Component->Thickness = GizmoLineThickness;
 		Component->NotifyExternalPropertyUpdates();
@@ -215,21 +402,40 @@ ACombinedTransformGizmoActor* ACombinedTransformGizmoActor::ConstructCustom3Axis
 	};
 	if ((Elements & ETransformGizmoSubElements::TranslateAxisX) != ETransformGizmoSubElements::None)
 	{
-		NewActor->TranslateX = MakeAxisArrowFunc(FLinearColor::Red, FVector(1, 0, 0));
+		NewActor->TranslateX = MakeAxisArrowFunc(EAxis::X);
 	}
 	if ((Elements & ETransformGizmoSubElements::TranslateAxisY) != ETransformGizmoSubElements::None)
 	{
-		NewActor->TranslateY = MakeAxisArrowFunc(FLinearColor::Green, FVector(0, 1, 0));
+		NewActor->TranslateY = MakeAxisArrowFunc(EAxis::Y);
 	}
 	if ((Elements & ETransformGizmoSubElements::TranslateAxisZ) != ETransformGizmoSubElements::None)
 	{
-		NewActor->TranslateZ = MakeAxisArrowFunc(FLinearColor::Blue, FVector(0, 0, 1));
+		NewActor->TranslateZ = MakeAxisArrowFunc(EAxis::Z);
 	}
 
 
-	auto MakePlaneRectFunc = [&](const FLinearColor& Color, const FVector& Axis0, const FVector& Axis1)
+	auto MakePlaneRectFunc = [World, NewActor, GizmoViewContext, GizmoLineThickness, &AddMeshGizmoComponent](EAxis::Type ElementAxis) -> UPrimitiveComponent*
 	{
-		UGizmoRectangleComponent* Component = AddDefaultRectangleComponent(World, NewActor, GizmoViewContext, Color, Axis0, Axis1);
+		if (GizmoDrawMode == DrawModeValue_Meshes)
+		{
+			if (UViewAdjustedStaticMeshGizmoComponent* Component = AddMeshGizmoComponent(
+				TEXT("/Engine/InteractiveToolsFramework/Meshes/GizmoPlaneHandle"),
+				UE::GizmoRenderingUtil::GetDefaultAxisColor(ElementAxis),
+				UE::GizmoUtil::GetRotatedBasisTransform(
+					// Transform for the X axis, relative to gizmo root
+					FTransform(FQuat::Identity, FVector(0, 40, 40), FVector::One()), ElementAxis),
+				EMirror::WhenCombined))
+			{
+				Component->SetRenderVisibilityFunction(ShouldPlaneBeVisible);
+				return Component;
+			}
+		}
+
+		// If we got to here, then we're creating the PDI drawn rectangle component
+		FVector AxisX, AxisY;
+		AxisToLegacyPairOfVectors(ElementAxis, AxisX, AxisY); 
+		UGizmoRectangleComponent* Component = AddDefaultRectangleComponent(World, NewActor, GizmoViewContext, 
+			AxisToLegacyColor(ElementAxis), AxisX, AxisY);
 		Component->LengthX = Component->LengthY = 30.0f;
 		Component->SegmentFlags = 0x2 | 0x4;
 		Component->Thickness = GizmoLineThickness;
@@ -238,20 +444,79 @@ ACombinedTransformGizmoActor* ACombinedTransformGizmoActor::ConstructCustom3Axis
 	};
 	if ((Elements & ETransformGizmoSubElements::TranslatePlaneYZ) != ETransformGizmoSubElements::None)
 	{
-		NewActor->TranslateYZ = MakePlaneRectFunc(FLinearColor::Red, FVector(0, 1, 0), FVector(0, 0, 1));
+		NewActor->TranslateYZ = MakePlaneRectFunc(EAxis::X);
 	}
 	if ((Elements & ETransformGizmoSubElements::TranslatePlaneXZ) != ETransformGizmoSubElements::None)
 	{
-		NewActor->TranslateXZ = MakePlaneRectFunc(FLinearColor::Green, FVector(1, 0, 0), FVector(0, 0, 1));
+		NewActor->TranslateXZ = MakePlaneRectFunc(EAxis::Y);
 	}
 	if ((Elements & ETransformGizmoSubElements::TranslatePlaneXY) != ETransformGizmoSubElements::None)
 	{
-		NewActor->TranslateXY = MakePlaneRectFunc(FLinearColor::Blue, FVector(1, 0, 0), FVector(0, 1, 0));
+		NewActor->TranslateXY = MakePlaneRectFunc(EAxis::Z);
+	}
+	if ((Elements & ETransformGizmoSubElements::FreeTranslate) != ETransformGizmoSubElements::None)
+	{
+		NewActor->FreeTranslateHandle = nullptr;
+		if (GizmoDrawMode == DrawModeValue_Meshes)
+		{
+			NewActor->FreeTranslateHandle = AddMeshGizmoComponent(
+				TEXT("/Engine/InteractiveToolsFramework/Meshes/GizmoSphereHandle"),
+				FreeTranslateColor,
+				FTransform::Identity,
+				EMirror::Never);
+		}
+
+		if (!NewActor->FreeTranslateHandle)
+		{
+			float BoxSize = 20.0f;
+			// We use a box as the backup because it already has hittesting for the inside, unlike our circles
+			NewActor->FreeTranslateHandle = AddDefaultBoxComponent(World, NewActor, GizmoViewContext, FLinearColor::Gray,
+				FVector(BoxSize / 2, BoxSize / 2, BoxSize / 2), FVector(BoxSize, BoxSize, BoxSize));
+		}
 	}
 
-	auto MakeAxisRotateCircleFunc = [&](const FLinearColor& Color, const FVector& Axis)
+	auto MakeAxisRotateCircleFunc = [World, NewActor, GizmoViewContext, GizmoLineThickness, &AddMeshGizmoComponent](EAxis::Type ElementAxis) -> UPrimitiveComponent*
 	{
-		UGizmoCircleComponent* Component = AddDefaultCircleComponent(World, NewActor, GizmoViewContext, Color, Axis, 120.0f);
+		if (GizmoDrawMode == DrawModeValue_Meshes)
+		{
+			FLinearColor Color = UE::GizmoRenderingUtil::GetDefaultAxisColor(ElementAxis);
+			Color.A = 0.75f; // Partially transparent, like editor gizmo
+
+			if (UViewAdjustedStaticMeshGizmoComponent* Component = AddMeshGizmoComponent(
+				TEXT("/Engine/InteractiveToolsFramework/Meshes/GizmoQuarterCircleHandle"),
+				Color,
+				UE::GizmoUtil::GetRotatedBasisTransform(
+					FTransform(FQuat::Identity, FVector::ZeroVector, FVector::OneVector), ElementAxis),
+				EMirror::Always))
+			{
+				Component->SetRenderVisibilityFunction(ShouldPlaneBeVisible);
+
+				UStaticMesh* SubstituteMesh = LoadObject<UStaticMesh>(nullptr, TEXT("/Engine/InteractiveToolsFramework/Meshes/GizmoFullCircleHandle"));
+				if (SubstituteMesh)
+				{
+					
+					UViewAdjustedStaticMeshGizmoComponent* SubstituteComponent = UE::GizmoRenderingUtil::CreateDefaultMaterialGizmoMeshComponent(
+						SubstituteMesh, GizmoViewContext, Component, Color,
+						// No need for hover material
+						false);
+					if (SubstituteComponent)
+					{
+						Component->SetSubstituteInteractionComponent(SubstituteComponent);
+
+						UE::GizmoRenderingUtil::FSubGizmoTransformAdjuster::AddTransformAdjuster(
+							SubstituteComponent,
+							// Don't need to keep distance to gizmo because we should be at gizmo root
+							nullptr, 
+							/*bMirror*/ false);
+					}
+				}
+
+				return Component;
+			}
+		}
+
+		UGizmoCircleComponent* Component = AddDefaultCircleComponent(World, NewActor, GizmoViewContext, 
+			AxisToLegacyColor(ElementAxis), AxisToVector(ElementAxis), 120.0f);
 		Component->Thickness = GizmoLineThickness;
 		Component->NotifyExternalPropertyUpdates();
 		return Component;
@@ -260,23 +525,23 @@ ACombinedTransformGizmoActor* ACombinedTransformGizmoActor::ConstructCustom3Axis
 	bool bAnyRotate = false;
 	if ((Elements & ETransformGizmoSubElements::RotateAxisX) != ETransformGizmoSubElements::None)
 	{
-		NewActor->RotateX = MakeAxisRotateCircleFunc(FLinearColor::Red, FVector(1, 0, 0));
+		NewActor->RotateX = MakeAxisRotateCircleFunc(EAxis::X);
 		bAnyRotate = true;
 	}
 	if ((Elements & ETransformGizmoSubElements::RotateAxisY) != ETransformGizmoSubElements::None)
 	{
-		NewActor->RotateY = MakeAxisRotateCircleFunc(FLinearColor::Green, FVector(0, 1, 0));
+		NewActor->RotateY = MakeAxisRotateCircleFunc(EAxis::Y);
 		bAnyRotate = true;
 	}
 	if ((Elements & ETransformGizmoSubElements::RotateAxisZ) != ETransformGizmoSubElements::None)
 	{
-		NewActor->RotateZ = MakeAxisRotateCircleFunc(FLinearColor::Blue, FVector(0, 0, 1));
+		NewActor->RotateZ = MakeAxisRotateCircleFunc(EAxis::Z);
 		bAnyRotate = true;
 	}
 
 
 	// add a non-interactive view-aligned circle element, so the axes look like a sphere.
-	if (bAnyRotate)
+	if (bAnyRotate && GizmoDrawMode != DrawModeValue_Meshes)
 	{
 		UGizmoCircleComponent* SphereEdge = NewObject<UGizmoCircleComponent>(NewActor);
 		NewActor->AddInstanceComponent(SphereEdge);
@@ -290,21 +555,88 @@ ACombinedTransformGizmoActor* ACombinedTransformGizmoActor::ConstructCustom3Axis
 		NewActor->RotationSphere = SphereEdge;
 	}
 
+	if ((Elements & ETransformGizmoSubElements::FreeRotate) != ETransformGizmoSubElements::None)
+	{
+		NewActor->FreeRotateHandle = nullptr;
+		if (GizmoDrawMode == DrawModeValue_Meshes)
+		{
+			NewActor->FreeRotateHandle = AddMeshGizmoComponent(
+				TEXT("/Engine/InteractiveToolsFramework/Meshes/GizmoSphereHandle"),
+				FreeRotateColor,
+				FTransform(FQuat::Identity, FVector::ZeroVector, FVector(9)),
+				EMirror::Never, /*bAddHoverMaterial*/ false);
+		}
+
+		if (!NewActor->FreeRotateHandle)
+		{
+			float BoxSize = 20.0f;
+			// We use a box as the backup because it already has hittesting for the inside, unlike our circles
+			NewActor->FreeRotateHandle = AddDefaultBoxComponent(World, NewActor, GizmoViewContext, FLinearColor::Gray,
+				FVector(BoxSize/2, BoxSize/2, BoxSize/2), FVector(BoxSize, BoxSize, BoxSize));
+		}
+	}
 
 
 	if ((Elements & ETransformGizmoSubElements::ScaleUniform) != ETransformGizmoSubElements::None)
 	{
-		float BoxSize = 14.0f;
-		UGizmoBoxComponent* ScaleComponent = AddDefaultBoxComponent(World, NewActor, GizmoViewContext, FLinearColor::Black,
-			FVector(BoxSize/2, BoxSize/2, BoxSize/2), FVector(BoxSize, BoxSize, BoxSize));
-		NewActor->UniformScale = ScaleComponent;
+		NewActor->UniformScale = nullptr;
+		if (GizmoDrawMode == DrawModeValue_Meshes)
+		{
+			NewActor->UniformScale = AddMeshGizmoComponent(
+				TEXT("/Engine/InteractiveToolsFramework/Meshes/GizmoBoxHandle"),
+				UniformScaleColor,
+				FTransform::Identity,
+				EMirror::Never);
+		}
+
+		if (!NewActor->UniformScale)
+		{
+			float BoxSize = 20.0f;
+			UGizmoBoxComponent* ScaleComponent = AddDefaultBoxComponent(World, NewActor, GizmoViewContext, FLinearColor::Black,
+				FVector(BoxSize/2, BoxSize/2, BoxSize/2), FVector(BoxSize, BoxSize, BoxSize));
+			NewActor->UniformScale = ScaleComponent;
+		}
 	}
 
 
 
-	auto MakeAxisScaleFunc = [&](const FLinearColor& Color, const FVector& Axis0, const FVector& Axis1, bool bLockSinglePlane)
+	auto MakeAxisScaleFunc = [World, NewActor, GizmoViewContext, GizmoLineThickness, &AddMeshGizmoComponent](EAxis::Type ElementAxis, const FVector& PerpendicularAxis, bool bLockSinglePlane,
+		TObjectPtr<UPrimitiveComponent>& FullHandleOut) -> UPrimitiveComponent*
 	{
-		UGizmoRectangleComponent* ScaleComponent = AddDefaultRectangleComponent(World, NewActor, GizmoViewContext, Color, Axis0, Axis1);
+		if (GizmoDrawMode == DrawModeValue_Meshes)
+		{
+			if (GizmoDrawMode == DrawModeValue_Meshes)
+			{
+				if (UViewAdjustedStaticMeshGizmoComponent* Component = AddMeshGizmoComponent(
+					TEXT("/Engine/InteractiveToolsFramework/Meshes/GizmoBoxHandle"),
+					UE::GizmoRenderingUtil::GetDefaultAxisColor(ElementAxis),
+					UE::GizmoUtil::GetRotatedBasisTransform(
+						FTransform(FQuat::Identity, FVector(130, 0, 0), FVector(0.8)),
+						ElementAxis),
+					EMirror::WhenCombined))
+				{
+					Component->SetRenderVisibilityFunction(ShouldAxisBeVisible);
+
+					// Also try to add a full handle to use when we're not using a combined gizmo.
+					FullHandleOut = AddMeshGizmoComponent(
+						TEXT("/Engine/InteractiveToolsFramework/Meshes/GizmoBoxArrowHandle"),
+						UE::GizmoRenderingUtil::GetDefaultAxisColor(ElementAxis),
+						UE::GizmoUtil::GetRotatedBasisTransform(
+							FTransform(FQuat::Identity, FVector::ZeroVector, FVector::One()),
+							ElementAxis),
+						EMirror::WhenCombined);
+					if (FullHandleOut)
+					{
+						FullHandleOut->SetVisibility(false);
+					}
+
+					return Component;
+				}
+			}
+		}
+
+		UGizmoRectangleComponent* ScaleComponent = AddDefaultRectangleComponent(World, NewActor, GizmoViewContext, 
+			AxisToLegacyColor(ElementAxis), AxisToVector(ElementAxis), PerpendicularAxis);
 		ScaleComponent->OffsetX = 140.0f; ScaleComponent->OffsetY = -10.0f;
 		ScaleComponent->LengthX = 7.0f; ScaleComponent->LengthY = 20.0f;
 		ScaleComponent->Thickness = GizmoLineThickness;
@@ -347,7 +679,8 @@ ACombinedTransformGizmoActor* ACombinedTransformGizmoActor::ConstructCustom3Axis
 		FVector PerpendicularAxis(0,1,0);
 		ConfigureAdditionalAxis(ETransformGizmoSubElements::ScaleAxisY, TotalAxisCount, PerpendicularAxis);
 		ConfigureAdditionalAxis(ETransformGizmoSubElements::ScaleAxisZ, TotalAxisCount, PerpendicularAxis);
-		NewActor->AxisScaleX = MakeAxisScaleFunc(FLinearColor::Red, FVector(1, 0, 0), PerpendicularAxis, TotalAxisCount == 2);
+		NewActor->AxisScaleX = MakeAxisScaleFunc(EAxis::X, PerpendicularAxis, TotalAxisCount == 2, 
+			NewActor->FullAxisScaleX);
 	}
 
 	if ((Elements & ETransformGizmoSubElements::ScaleAxisY) != ETransformGizmoSubElements::None)
@@ -356,7 +689,8 @@ ACombinedTransformGizmoActor* ACombinedTransformGizmoActor::ConstructCustom3Axis
 		FVector PerpendicularAxis(1, 0, 0);
 		ConfigureAdditionalAxis(ETransformGizmoSubElements::ScaleAxisX, TotalAxisCount, PerpendicularAxis);
 		ConfigureAdditionalAxis(ETransformGizmoSubElements::ScaleAxisZ, TotalAxisCount, PerpendicularAxis);
-		NewActor->AxisScaleY = MakeAxisScaleFunc(FLinearColor::Green, FVector(0, 1, 0), PerpendicularAxis, TotalAxisCount == 2);
+		NewActor->AxisScaleY = MakeAxisScaleFunc(EAxis::Y, PerpendicularAxis, TotalAxisCount == 2, 
+			NewActor->FullAxisScaleY);
 	}
 
 	if ((Elements & ETransformGizmoSubElements::ScaleAxisZ) != ETransformGizmoSubElements::None)
@@ -365,13 +699,35 @@ ACombinedTransformGizmoActor* ACombinedTransformGizmoActor::ConstructCustom3Axis
 		FVector PerpendicularAxis(1, 0, 0);
 		ConfigureAdditionalAxis(ETransformGizmoSubElements::ScaleAxisY, TotalAxisCount, PerpendicularAxis);
 		ConfigureAdditionalAxis(ETransformGizmoSubElements::ScaleAxisX, TotalAxisCount, PerpendicularAxis);
-		NewActor->AxisScaleZ = MakeAxisScaleFunc(FLinearColor::Blue, FVector(0, 0, 1), PerpendicularAxis, TotalAxisCount == 2);
+		NewActor->AxisScaleZ = MakeAxisScaleFunc(EAxis::Z, PerpendicularAxis, TotalAxisCount == 2, 
+			NewActor->FullAxisScaleZ);
 	}
 
 
-	auto MakePlaneScaleFunc = [&](const FLinearColor& Color, const FVector& Axis0, const FVector& Axis1)
+	auto MakePlaneScaleFunc = [World, NewActor, GizmoViewContext, GizmoLineThickness, &AddMeshGizmoComponent](EAxis::Type ElementAxis) -> UPrimitiveComponent*
 	{
-		UGizmoRectangleComponent* ScaleComponent = AddDefaultRectangleComponent(World, NewActor, GizmoViewContext, Color, Axis0, Axis1);
+		if (GizmoDrawMode == DrawModeValue_Meshes)
+		{
+			if (UViewAdjustedStaticMeshGizmoComponent* Component = AddMeshGizmoComponent(
+				TEXT("/Engine/InteractiveToolsFramework/Meshes/GizmoCornerHandle"),
+				UE::GizmoRenderingUtil::GetDefaultAxisColor(ElementAxis),
+				UE::GizmoUtil::GetRotatedBasisTransform(
+					// Transform for the X axis, relative to gizmo root
+					FTransform(FQuat::Identity, CornerScalePositionCombined, FVector(CornerScaleHandleScale)), ElementAxis),
+				// We actually adjust the transform of the plane scale handles and swap the adjuster inside
+				//  ApplyGizmoActiveMode, so we don't need this one to be updated.
+				EMirror::Always))
+			{
+				Component->SetRenderVisibilityFunction(ShouldPlaneBeVisible);
+				return Component;
+			}
+		}
+
+		// If we got to here, then we're creating the PDI drawn rectangle component
+		FVector Axis0, Axis1;
+		AxisToLegacyPairOfVectors(ElementAxis, Axis0, Axis1);
+		UGizmoRectangleComponent* ScaleComponent = AddDefaultRectangleComponent(World, NewActor, GizmoViewContext,
+			AxisToLegacyColor(ElementAxis), Axis0, Axis1);
 		ScaleComponent->OffsetX = ScaleComponent->OffsetY = 120.0f;
 		ScaleComponent->LengthX = ScaleComponent->LengthY = 20.0f;
 		ScaleComponent->Thickness = GizmoLineThickness;
@@ -381,15 +737,15 @@ ACombinedTransformGizmoActor* ACombinedTransformGizmoActor::ConstructCustom3Axis
 	};
 	if ((Elements & ETransformGizmoSubElements::ScalePlaneYZ) != ETransformGizmoSubElements::None)
 	{
-		NewActor->PlaneScaleYZ = MakePlaneScaleFunc(FLinearColor::Red, FVector(0, 1, 0), FVector(0, 0, 1));
+		NewActor->PlaneScaleYZ = MakePlaneScaleFunc(EAxis::X);
 	}
 	if ((Elements & ETransformGizmoSubElements::ScalePlaneXZ) != ETransformGizmoSubElements::None)
 	{
-		NewActor->PlaneScaleXZ = MakePlaneScaleFunc(FLinearColor::Green, FVector(1, 0, 0), FVector(0, 0, 1));
+		NewActor->PlaneScaleXZ = MakePlaneScaleFunc(EAxis::Y);
 	}
 	if ((Elements & ETransformGizmoSubElements::ScalePlaneXY) != ETransformGizmoSubElements::None)
 	{
-		NewActor->PlaneScaleXY = MakePlaneScaleFunc(FLinearColor::Blue, FVector(1, 0, 0), FVector(0, 1, 0));
+		NewActor->PlaneScaleXY = MakePlaneScaleFunc(EAxis::Z);
 	}
 
 
@@ -603,9 +959,19 @@ bool UCombinedTransformGizmo::SetSubGizmoComponent(ETransformGizmoSubElements El
 		AddUniformScaleGizmo(Params, *SubGizmoSharedState);
 		break;
 	}
+	case ETransformGizmoSubElements::FreeTranslate:
+	{
+		AddFreeTranslationGizmo(Params, *SubGizmoSharedState);
+		break;
+	}
+	case ETransformGizmoSubElements::FreeRotate:
+	{
+		AddFreeRotationGizmo(Params, *SubGizmoSharedState);
+		break;
+	}
 	case ETransformGizmoSubElements::RotateAllAxes:
 	{
-		// no gizmo for the sphere currently
+		// no gizmo for the drawn sphere
 		if (ensure(GizmoActor->RotationSphere == NewComponent))
 		{
 			ActiveComponents.Add(GizmoActor->RotationSphere);
@@ -656,7 +1022,8 @@ void UCombinedTransformGizmo::SetWorldAlignmentFunctions(
 			{
 				if (CastHitTarget->Component == GizmoActor->TranslateXY
 					|| CastHitTarget->Component == GizmoActor->TranslateXZ
-					|| CastHitTarget->Component == GizmoActor->TranslateYZ)
+					|| CastHitTarget->Component == GizmoActor->TranslateYZ
+					|| CastHitTarget->Component == GizmoActor->FreeTranslateHandle)
 				{
 					CastGizmo->ShouldUseCustomDestinationFunc = [this]() { return ShouldAlignDestination(); };
 					CastGizmo->CustomDestinationFunc =
@@ -735,9 +1102,9 @@ void UCombinedTransformGizmo::Setup()
 	}
 
 	GizmoActor = GizmoActorBuilder->CreateNewGizmoActor(World);
-}
-
 
+	PreviousActiveGizmoMode = ActiveGizmoMode;
+}
 
 void UCombinedTransformGizmo::Shutdown()
 {
@@ -754,26 +1121,22 @@ void UCombinedTransformGizmo::Shutdown()
 
 void UCombinedTransformGizmo::UpdateCameraAxisSource()
 {
-	FViewCameraState CameraState;
-	GetGizmoManager()->GetContextQueriesAPI()->GetCurrentViewState(CameraState);
 	if (CameraAxisSource != nullptr && GizmoActor != nullptr)
 	{
-		CameraAxisSource->Origin = GizmoActor->GetTransform().GetLocation();
-		CameraAxisSource->Direction = -CameraState.Forward();
-		CameraAxisSource->TangentX = CameraState.Right();
-		CameraAxisSource->TangentY = CameraState.Up();
+		UE::GizmoUtil::UpdateCameraAxisSource(*CameraAxisSource, GetGizmoManager(), 
+			GizmoActor->GetTransform().GetLocation());
 	}
 }
 
 
 void UCombinedTransformGizmo::Tick(float DeltaTime)
-{	
+{
 	if (bUseContextCoordinateSystem)
 	{
 		CurrentCoordinateSystem = GetGizmoManager()->GetContextQueriesAPI()->GetCurrentCoordinateSystem();
 	}
 	check(CurrentCoordinateSystem == EToolContextCoordinateSystem::World || CurrentCoordinateSystem == EToolContextCoordinateSystem::Local)
-	
+
 	FToolContextSnappingConfiguration SnappingConfig = GetGizmoManager()->GetContextQueriesAPI()->GetCurrentSnappingSettings();
 	RelativeTranslationSnapping.UpdateContextValue(SnappingConfig.bEnableAbsoluteWorldSnapping == false);
 
@@ -792,7 +1155,23 @@ void UCombinedTransformGizmo::Tick(float DeltaTime)
 		}
 	}
 
+	if (bUseContextGizmoMode)
+	{
+		ActiveGizmoMode = GetGizmoManager()->GetContextQueriesAPI()->GetCurrentTransformGizmoMode();
+	}
+
 	// apply dynamic visibility filtering to sub-gizmos
+	if (PreviousActiveGizmoMode != ActiveGizmoMode)
+	{
+		ApplyGizmoActiveMode();
+	}
+
+	UpdateCameraAxisSource();
+}
+
+void UCombinedTransformGizmo::ApplyGizmoActiveMode()
+{
+	using namespace CombinedTransformGizmoLocals;
 
 	auto SetSubGizmoTypeVisibility = [this](TArray<FSubGizmoInfo>& GizmoInfos, bool bVisible)
 	{
@@ -804,29 +1183,107 @@ void UCombinedTransformGizmo::Tick(float DeltaTime)
 			}
 		}
 	};
-
-	if (bUseContextGizmoMode)
-	{
-		ActiveGizmoMode = GetGizmoManager()->GetContextQueriesAPI()->GetCurrentTransformGizmoMode();
-	}
-	EToolContextTransformGizmoMode UseGizmoMode = ActiveGizmoMode;
-
 	bool bShouldShowTranslation =
-		(UseGizmoMode == EToolContextTransformGizmoMode::Combined || UseGizmoMode == EToolContextTransformGizmoMode::Translation);
+		(ActiveGizmoMode == EToolContextTransformGizmoMode::Combined || ActiveGizmoMode == EToolContextTransformGizmoMode::Translation);
 	bool bShouldShowRotation =
-		(UseGizmoMode == EToolContextTransformGizmoMode::Combined || UseGizmoMode == EToolContextTransformGizmoMode::Rotation);
+		(ActiveGizmoMode == EToolContextTransformGizmoMode::Combined || ActiveGizmoMode == EToolContextTransformGizmoMode::Rotation);
 	bool bShouldShowUniformScale =
-		(UseGizmoMode == EToolContextTransformGizmoMode::Combined || UseGizmoMode == EToolContextTransformGizmoMode::Scale);
+		(ActiveGizmoMode == EToolContextTransformGizmoMode::Combined || ActiveGizmoMode == EToolContextTransformGizmoMode::Scale);
 	bool bShouldShowNonUniformScale = 
-		(UseGizmoMode == EToolContextTransformGizmoMode::Combined || UseGizmoMode == EToolContextTransformGizmoMode::Scale)
+		(ActiveGizmoMode == EToolContextTransformGizmoMode::Combined || ActiveGizmoMode == EToolContextTransformGizmoMode::Scale)
 		&& IsNonUniformScaleAllowedFunc();
 
 	SetSubGizmoTypeVisibility(TranslationSubGizmos, bShouldShowTranslation);
 	SetSubGizmoTypeVisibility(RotationSubGizmos, bShouldShowRotation);
 	SetSubGizmoTypeVisibility(UniformScaleSubGizmos, bShouldShowUniformScale);
-	SetSubGizmoTypeVisibility(NonUniformScaleSubGizmos, bShouldShowNonUniformScale);
 
-	UpdateCameraAxisSource();
+	if (ActiveGizmoMode == EToolContextTransformGizmoMode::Combined || ActiveGizmoMode == EToolContextTransformGizmoMode::Scale)
+	{
+		// The scale handles look different in different modes, so swap them if necessary
+		auto SwapAxisScaleComponent = [this](UPrimitiveComponent* HandleInCombined, UPrimitiveComponent* HandleInSeparate)
+		{
+			UPrimitiveComponent* HandleToUse = ActiveGizmoMode == EToolContextTransformGizmoMode::Combined ? HandleInCombined : HandleInSeparate;
+			UPrimitiveComponent* HandleToReplace = ActiveGizmoMode == EToolContextTransformGizmoMode::Combined ? HandleInSeparate : HandleInCombined;
+
+			// Don't swap if we don't have an alternative to use
+			if (!HandleToUse || !HandleToReplace) return;
+				
+			FSubGizmoInfo* GizmoInfo = NonUniformScaleSubGizmos.FindByPredicate([HandleToUse, HandleToReplace](const FSubGizmoInfo& Info)
+			{ 
+				return Info.Component == HandleToUse || Info.Component == HandleToReplace;
+			});
+			// Don't swap if we don't have this gizmo or if it's already using the correct one
+			if (!GizmoInfo || GizmoInfo->Component == HandleToUse) return;
+
+			UAxisPositionGizmo* SubGizmo = Cast<UAxisPositionGizmo>(GizmoInfo->Gizmo);
+			if (!ensure(SubGizmo)) return;
+			UGizmoComponentHitTarget* HitTarget = Cast<UGizmoComponentHitTarget>(SubGizmo->HitTarget.GetObject());
+			if (!ensure(HitTarget)) return;
+
+			HandleToReplace->SetVisibility(false);
+			HitTarget->Component = HandleToUse;
+			GizmoInfo->Component = HandleToUse;
+		};
+
+		SwapAxisScaleComponent(GizmoActor->AxisScaleX, GizmoActor->FullAxisScaleX);
+		SwapAxisScaleComponent(GizmoActor->AxisScaleY, GizmoActor->FullAxisScaleY);
+		SwapAxisScaleComponent(GizmoActor->AxisScaleZ, GizmoActor->FullAxisScaleZ);
+
+		// The plane scale handles look better if they are closer to the gizmo when not combined.
+		auto AdjustPlaneScaleComponent = [this](UViewAdjustedStaticMeshGizmoComponent* Component, EAxis::Type ElementAxis)
+		{
+			if (!Component) return;
+			Component->SetRelativeTransform(UE::GizmoUtil::GetRotatedBasisTransform(
+				// Transform for the X axis, relative to gizmo root
+				FTransform(FQuat::Identity, ActiveGizmoMode == EToolContextTransformGizmoMode::Combined ? CornerScalePositionCombined : CornerScalePositionSeparate,
+					FVector(CornerScaleHandleScale)), ElementAxis));
+
+			// Just replace the adjuster
+			// TODO: Maybe keep track of whether this needs doing instead of doing it each time we switch to scale mode
+			TSharedPtr<UE::GizmoRenderingUtil::FSubGizmoTransformAdjuster> Adjuster =
+				UE::GizmoRenderingUtil::FSubGizmoTransformAdjuster::AddTransformAdjuster(
+					Component, GizmoActor->GetRootComponent(), ActiveGizmoMode == EToolContextTransformGizmoMode::Combined);
+		};
+		AdjustPlaneScaleComponent(Cast<UViewAdjustedStaticMeshGizmoComponent>(GizmoActor->PlaneScaleXY), EAxis::Z);
+		AdjustPlaneScaleComponent(Cast<UViewAdjustedStaticMeshGizmoComponent>(GizmoActor->PlaneScaleXZ), EAxis::Y);
+		AdjustPlaneScaleComponent(Cast<UViewAdjustedStaticMeshGizmoComponent>(GizmoActor->PlaneScaleYZ), EAxis::X);
+	}
+	// This is done after the above, since the above affects NonUniformScaleSubGizmos
+	SetSubGizmoTypeVisibility(NonUniformScaleSubGizmos, bShouldShowUniformScale);
+
+	auto SetVisibilityIfExists = [](UPrimitiveComponent* Component, bool bVisible)
+	{
+		if (Component)
+		{
+			Component->SetVisibility(bVisible);
+		}
+	};
+	if (GizmoActor->FreeRotateHandle)
+	{
+		GizmoActor->FreeRotateHandle->SetVisibility(ActiveGizmoMode == EToolContextTransformGizmoMode::Rotation);
+	}
+	if (GizmoActor->FreeTranslateHandle)
+	{
+		GizmoActor->FreeTranslateHandle->SetVisibility(ActiveGizmoMode == EToolContextTransformGizmoMode::Translation
+			|| (ActiveGizmoMode == EToolContextTransformGizmoMode::Combined && UniformScaleSubGizmos.Num() == 0));
+	}
+
+	if (GizmoActor && GizmoDrawMode == DrawModeValue_Meshes)
+	{
+		// Many components mirror in combined mode but not separate mode, mostly because it is weird for them
+		//  to not mirror when the rotation components do.
+		for (int32 i = GizmoActor->AdjustersThatMirrorOnlyInCombinedMode.Num() - 1; i >= 0; --i)
+		{
+			TSharedPtr<UE::GizmoRenderingUtil::FSubGizmoTransformAdjuster> Adjuster =
+				GizmoActor->AdjustersThatMirrorOnlyInCombinedMode[i].Pin();
+			if (ensure(Adjuster))
+			{
+				Adjuster->SetMirrorBasedOnOctant(ActiveGizmoMode == EToolContextTransformGizmoMode::Combined);
+			}
+		}
+	}
+
+	PreviousActiveGizmoMode = ActiveGizmoMode;
 }
 
 
@@ -877,6 +1334,12 @@ void UCombinedTransformGizmo::SetActiveTarget(UTransformProxy* Target, IToolCont
 			AddPlaneTranslationGizmo(Params, *SubGizmoSharedState);
 		}
 	}
+	if (GizmoActor->FreeTranslateHandle != nullptr)
+	{
+		Params.Component = GizmoActor->FreeTranslateHandle;
+		Params.Axis = EAxis::None;
+		AddFreeTranslationGizmo(Params, *SubGizmoSharedState);
+	}
 	UPrimitiveComponent* RotationAxisComponents[3]{ GizmoActor->RotateX, GizmoActor->RotateY, GizmoActor->RotateZ };
 	for (int AxisIndex = 0; AxisIndex < 3; ++AxisIndex)
 	{
@@ -889,10 +1352,15 @@ void UCombinedTransformGizmo::SetActiveTarget(UTransformProxy* Target, IToolCont
 	}
 	if (GizmoActor->RotationSphere != nullptr)
 	{
-		// no gizmo for the sphere currently
 		ActiveComponents.Add(GizmoActor->RotationSphere);
 		RotationSubGizmos.Add(FSubGizmoInfo{ GizmoActor->RotationSphere, nullptr });
 	}
+	if (GizmoActor->FreeRotateHandle != nullptr)
+	{
+		Params.Component = GizmoActor->FreeRotateHandle;
+		Params.Axis = EAxis::None;
+		AddFreeRotationGizmo(Params, *SubGizmoSharedState);
+	}
 	if (GizmoActor->UniformScale != nullptr)
 	{
 		Params.Component = GizmoActor->UniformScale;
@@ -931,6 +1399,8 @@ void UCombinedTransformGizmo::SetActiveTarget(UTransformProxy* Target, IToolCont
 	UnitAxisYSource = SubGizmoSharedState->UnitCardinalAxisSources[1];
 	UnitAxisZSource = SubGizmoSharedState->UnitCardinalAxisSources[2];
 
+	ApplyGizmoActiveMode();
+
 	OnSetActiveTarget.Broadcast(this, ActiveTarget);
 }
 
@@ -1211,6 +1681,41 @@ UInteractiveGizmo* UCombinedTransformGizmo::AddUniformScaleGizmo(
 	return Gizmo;
 }
 
+UInteractiveGizmo* UCombinedTransformGizmo::AddFreeTranslationGizmo(
+	FTransformSubGizmoCommonParams& Params, FTransformSubGizmoSharedState& SharedState)
+{
+	UFreePositionSubGizmo* Gizmo = UE::GizmoUtil::CreateGizmoViaSimpleBuilder<UFreePositionSubGizmo>(
+		GetGizmoManager(), FString(), this);
+	if (!ensure(Gizmo))
+	{
+		return nullptr;
+	}
+	ensure(Gizmo->InitializeAsScreenPlaneTranslateGizmo(Params, &SharedState));
+
+	ActiveComponents.Add(Params.Component);
+	ActiveGizmos.Add(Gizmo);
+
+	return Gizmo;
+}
+
+UInteractiveGizmo* UCombinedTransformGizmo::AddFreeRotationGizmo(
+	FTransformSubGizmoCommonParams& Params, FTransformSubGizmoSharedState& SharedState)
+{
+	UFreeRotationSubGizmo* Gizmo = UE::GizmoUtil::CreateGizmoViaSimpleBuilder<UFreeRotationSubGizmo>(
+		GetGizmoManager(), FString(), this);
+	if (!ensure(Gizmo))
+	{
+		return nullptr;
+	}
+	UGizmoViewContext* GizmoViewContext = UE::GizmoUtil::GetGizmoViewContext(GetGizmoManager());
+	ensure(Gizmo->InitializeAsRotationGizmo(Params, GizmoViewContext, &SharedState));
+
+	ActiveComponents.Add(Params.Component);
+	ActiveGizmos.Add(Gizmo);
+
+	return Gizmo;
+}
+
 
 // These are deprecated initialization functions that do sub gizmo initialization by hand instead of using the
 //  "InitializeAs..." functions that were added to subgizmos to make them simpler to instantiate outside of
diff --git a/Engine/Source/Runtime/InteractiveToolsFramework/Private/BaseGizmos/RepositionableTransformGizmo.cpp b/Engine/Source/Runtime/InteractiveToolsFramework/Private/BaseGizmos/RepositionableTransformGizmo.cpp
--- a/Engine/Source/Runtime/InteractiveToolsFramework/Private/BaseGizmos/RepositionableTransformGizmo.cpp
+++ b/Engine/Source/Runtime/InteractiveToolsFramework/Private/BaseGizmos/RepositionableTransformGizmo.cpp
@@ -5,6 +5,8 @@
 #include "BaseBehaviors/ClickDragBehavior.h"
 #include "BaseGizmos/AxisAngleGizmo.h"
 #include "BaseGizmos/AxisPositionGizmo.h"
+#include "BaseGizmos/FreePositionSubGizmo.h"
+#include "BaseGizmos/FreeRotationSubGizmo.h"
 #include "BaseGizmos/GizmoViewContext.h"
 #include "BaseGizmos/HitTargets.h"
 #include "BaseGizmos/PlanePositionGizmo.h"
@@ -132,6 +134,23 @@ void URepositionableTransformGizmo::SetActiveTarget(UTransformProxy* Target, ITo
 			PivotAlignmentGizmos.Add(SubGizmo);
 		}
 	}
+
+	Params.Axis = EAxis::None;
+	if (GizmoActor->FreeTranslateHandle)
+	{
+		Params.Component = GizmoActor->FreeTranslateHandle;
+		UInteractiveGizmo* SubGizmo = AddFreeTranslationGizmo(Params, SharedStateToUse);
+		ModifyPivotFreeTranslationGizmo(SubGizmo);
+		PivotAlignmentGizmos.Add(SubGizmo);
+	}
+	
+	if (GizmoActor->FreeRotateHandle)
+	{
+		Params.Component = GizmoActor->FreeRotateHandle;
+		UInteractiveGizmo* SubGizmo = AddFreeRotationGizmo(Params, SharedStateToUse);
+		ModifyPivotFreeRotationGizmo(SubGizmo);
+		PivotAlignmentGizmos.Add(SubGizmo);
+	}
 }
 
 /** 
@@ -202,6 +221,41 @@ void URepositionableTransformGizmo::ModifyPivotRotateGizmo(UInteractiveGizmo* Su
 	};
 }
 
+void URepositionableTransformGizmo::ModifyPivotFreeTranslationGizmo(UInteractiveGizmo* SubGizmoIn)
+{
+	UFreePositionSubGizmo* SubGizmo = Cast<UFreePositionSubGizmo>(SubGizmoIn);
+	if (!ensure(SubGizmo)) { return; }
+
+	SubGizmo->MouseBehavior->SetUseMiddleMouseButton();
+	SubGizmo->ShouldUseCustomDestinationFunc = [this]() { return ShouldAlignPivot(); };
+	SubGizmo->CustomDestinationFunc =
+		[this](const UPlanePositionGizmo::FCustomDestinationParams& Params, FVector& OutputPoint) {
+		return PivotAlignmentRayCaster(*Params.WorldRay, OutputPoint);
+	};
+
+	UGizmoComponentHitTarget* HitTarget = Cast<UGizmoComponentHitTarget>(SubGizmo->HitTarget.GetObject());
+	check(HitTarget);
+	HitTarget->UpdateHoverFunction = nullptr; // Already taken care of by the normal gizmo
+}
+
+void URepositionableTransformGizmo::ModifyPivotFreeRotationGizmo(UInteractiveGizmo* SubGizmoIn)
+{
+	UFreeRotationSubGizmo* SubGizmo = Cast<UFreeRotationSubGizmo>(SubGizmoIn);
+	if (!ensure(SubGizmo)) { return; }
+
+	SubGizmo->MouseBehavior->SetUseMiddleMouseButton();
+
+	UGizmoComponentHitTarget* HitTarget = Cast<UGizmoComponentHitTarget>(SubGizmo->HitTarget.GetObject());
+	check(HitTarget);
+	HitTarget->UpdateHoverFunction = nullptr; // Already taken care of by the normal gizmo
+
+	// We don't want to be able to reposition the rotation of the gizmo when using the world coordinate system,
+	// since the axes must stay aligned with world axes. 
+	HitTarget->Condition = [this](const FInputDeviceRay&) {
+		return (CurrentCoordinateSystem == EToolContextCoordinateSystem::Local);
+	};
+}
+
 void URepositionableTransformGizmo::SetWorldAlignmentFunctions(
 	TUniqueFunction<bool()>&& ShouldAlignTranslationIn,
 	TUniqueFunction<bool(const FRay&, FVector&)>&& TranslationAlignmentRayCasterIn)
@@ -253,6 +307,14 @@ void URepositionableTransformGizmo::SetPivotAlignmentFunctions(
 				return PivotAlignmentRayCaster(*Params.WorldRay, OutputPoint);
 			};
 		}
+		if (UFreePositionSubGizmo* CastGizmo = Cast<UFreePositionSubGizmo>(SubGizmo))
+		{
+			CastGizmo->ShouldUseCustomDestinationFunc = [this]() { return ShouldAlignPivot(); };
+			CastGizmo->CustomDestinationFunc =
+				[this](const UFreePositionSubGizmo::FCustomDestinationParams& Params, FVector& OutputPoint) {
+				return PivotAlignmentRayCaster(*Params.WorldRay, OutputPoint);
+			};
+		}
 	}
 }
 
diff --git a/Engine/Source/Runtime/InteractiveToolsFramework/Public/BaseGizmos/CombinedTransformGizmo.h b/Engine/Source/Runtime/InteractiveToolsFramework/Public/BaseGizmos/CombinedTransformGizmo.h
--- a/Engine/Source/Runtime/InteractiveToolsFramework/Public/BaseGizmos/CombinedTransformGizmo.h
+++ b/Engine/Source/Runtime/InteractiveToolsFramework/Public/BaseGizmos/CombinedTransformGizmo.h
@@ -26,6 +26,10 @@ namespace UE::GizmoUtil
 	struct FTransformSubGizmoCommonParams;
 	struct FTransformSubGizmoSharedState;
 }
+namespace UE::GizmoRenderingUtil
+{
+	class FSubGizmoTransformAdjuster;
+}
 
 /**
  * ACombinedTransformGizmoActor is an Actor type intended to be used with UCombinedTransformGizmo,
@@ -98,6 +102,12 @@ public:
 	UPROPERTY()
 	TObjectPtr<UPrimitiveComponent> RotationSphere;
 
+	UPROPERTY()
+	TObjectPtr<UPrimitiveComponent> FreeRotateHandle;
+
+	UPROPERTY()
+	TObjectPtr<UPrimitiveComponent> FreeTranslateHandle;
+
 	//
 	// Scaling Components
 	//
@@ -167,6 +177,19 @@ public:
 		UWorld* World, UGizmoViewContext* GizmoViewContext,
 		ETransformGizmoSubElements Elements
 	);
+
+private:
+	TArray<TWeakPtr<UE::GizmoRenderingUtil::FSubGizmoTransformAdjuster>> AdjustersThatMirrorOnlyInCombinedMode;
+	friend class UCombinedTransformGizmo;
+	
+	// These store versions of the axis scale components that can be used when not using a combined gizmo,
+	//  if they differ.
+	UPROPERTY()
+	TObjectPtr<UPrimitiveComponent> FullAxisScaleX;
+	UPROPERTY()
+	TObjectPtr<UPrimitiveComponent> FullAxisScaleY;
+	UPROPERTY()
+	TObjectPtr<UPrimitiveComponent> FullAxisScaleZ;
 };
 
 
@@ -369,10 +392,9 @@ public:
 	/**
 	 * By default, the nonuniform scale components can scale negatively. However, they can be made to clamp
 	 * to zero instead by passing true here. This is useful for using the gizmo to flatten geometry.
-	 *
-	 * TODO: Should this affect uniform scaling too?
 	 */
 	INTERACTIVETOOLSFRAMEWORK_API virtual void SetDisallowNegativeScaling(bool bDisallow);
+	//~ TODO: Should the above affect uniform scaling too?
 
 	// UInteractiveGizmo overrides
 	INTERACTIVETOOLSFRAMEWORK_API virtual void Setup() override;
@@ -610,8 +632,12 @@ protected:
 	TArray<FSubGizmoInfo> RotationSubGizmos;
 	TArray<FSubGizmoInfo> UniformScaleSubGizmos;
 	TArray<FSubGizmoInfo> NonUniformScaleSubGizmos;
+private:
+	TWeakObjectPtr<UInteractiveGizmo> AxisScaleXGizmo;
+	TWeakObjectPtr<UInteractiveGizmo> AxisScaleYGizmo;
+	TWeakObjectPtr<UInteractiveGizmo> AxisScaleZGizmo;
 
-
+protected:
 	/** GizmoActors will be spawned in this World */
 	UWorld* World;
 
@@ -717,6 +743,12 @@ protected:
 	INTERACTIVETOOLSFRAMEWORK_API virtual UInteractiveGizmo* AddUniformScaleGizmo(
 		FTransformSubGizmoCommonParams& Params, FTransformSubGizmoSharedState& SharedState);
 
+	INTERACTIVETOOLSFRAMEWORK_API virtual UInteractiveGizmo* AddFreeTranslationGizmo(
+		FTransformSubGizmoCommonParams& Params, FTransformSubGizmoSharedState& SharedState);
+
+	INTERACTIVETOOLSFRAMEWORK_API virtual UInteractiveGizmo* AddFreeRotationGizmo(
+		FTransformSubGizmoCommonParams& Params, FTransformSubGizmoSharedState& SharedState);
+
 	// Axis and Plane TransformSources use these function to execute snapping queries
 	INTERACTIVETOOLSFRAMEWORK_API bool PositionSnapFunction(const FVector& WorldPosition, FVector& SnappedPositionOut) const;
 	INTERACTIVETOOLSFRAMEWORK_API bool PositionAxisDeltaSnapFunction(double AxisDelta, double& SnappedDeltaOut, int AxisIndex) const;
@@ -780,4 +812,7 @@ private:
 	// Here to support subgizmo reinitialization after SetActiveTarget has been called. Private
 	// instead of protected for now in case we change the approach here.
 	IToolContextTransactionProvider* TransactionProviderAtLastSetActiveTarget = nullptr;
+
+	EToolContextTransformGizmoMode PreviousActiveGizmoMode = EToolContextTransformGizmoMode::Combined;
+	void ApplyGizmoActiveMode();
 };
diff --git a/Engine/Source/Runtime/InteractiveToolsFramework/Public/BaseGizmos/RepositionableTransformGizmo.h b/Engine/Source/Runtime/InteractiveToolsFramework/Public/BaseGizmos/RepositionableTransformGizmo.h
--- a/Engine/Source/Runtime/InteractiveToolsFramework/Public/BaseGizmos/RepositionableTransformGizmo.h
+++ b/Engine/Source/Runtime/InteractiveToolsFramework/Public/BaseGizmos/RepositionableTransformGizmo.h
@@ -72,4 +72,7 @@ protected:
 	INTERACTIVETOOLSFRAMEWORK_API void ModifyPivotAxisGizmo(UInteractiveGizmo* SubGizmoIn);
 	INTERACTIVETOOLSFRAMEWORK_API void ModifyPivotPlaneGizmo(UInteractiveGizmo* SubGizmoIn);
 	INTERACTIVETOOLSFRAMEWORK_API void ModifyPivotRotateGizmo(UInteractiveGizmo* SubGizmoIn);
+private:
+	void ModifyPivotFreeTranslationGizmo(UInteractiveGizmo* SubGizmoIn);
+	void ModifyPivotFreeRotationGizmo(UInteractiveGizmo* SubGizmoIn);
 };
diff --git a/Engine/Source/Runtime/InteractiveToolsFramework/Public/InteractiveGizmo.h b/Engine/Source/Runtime/InteractiveToolsFramework/Public/InteractiveGizmo.h
--- a/Engine/Source/Runtime/InteractiveToolsFramework/Public/InteractiveGizmo.h
+++ b/Engine/Source/Runtime/InteractiveToolsFramework/Public/InteractiveGizmo.h
@@ -132,8 +132,14 @@ enum class ETransformGizmoSubElements
 
 	ScaleUniform = 1<<16,
 
-	StandardTranslateRotate = TranslateAllAxes | TranslateAllPlanes | RotateAllAxes,
-	TranslateRotateUniformScale = TranslateAllAxes | TranslateAllPlanes | RotateAllAxes | ScaleUniform,
-	FullTranslateRotateScale = TranslateAllAxes | TranslateAllPlanes | RotateAllAxes | ScaleAllAxes | ScaleAllPlanes | ScaleUniform
+	FreeTranslate = 1<<17,
+	FreeRotate = 1<<18,
+
+	StandardTranslateRotate = TranslateAllAxes | TranslateAllPlanes | RotateAllAxes 
+		| FreeTranslate | FreeRotate,
+	TranslateRotateUniformScale = TranslateAllAxes | TranslateAllPlanes | RotateAllAxes 
+		| ScaleUniform | FreeTranslate | FreeRotate,
+	FullTranslateRotateScale = TranslateAllAxes | TranslateAllPlanes | RotateAllAxes 
+		| ScaleAllAxes | ScaleAllPlanes | ScaleUniform | FreeTranslate | FreeRotate
 };
 ENUM_CLASS_FLAGS(ETransformGizmoSubElements);
