InteractiveToolsFramework: Fix static analysis warnings in CombinedTransformGizmo.

diff --git a/Engine/Source/Runtime/InteractiveToolsFramework/Private/BaseGizmos/CombinedTransformGizmo.cpp b/Engine/Source/Runtime/InteractiveToolsFramework/Private/BaseGizmos/CombinedTransformGizmo.cpp
--- a/Engine/Source/Runtime/InteractiveToolsFramework/Private/BaseGizmos/CombinedTransformGizmo.cpp
+++ b/Engine/Source/Runtime/InteractiveToolsFramework/Private/BaseGizmos/CombinedTransformGizmo.cpp
@@ -605,33 +605,30 @@ ACombinedTransformGizmoActor* ACombinedTransformGizmoActor::ConstructCustom3Axis
 	{
 		if (GizmoDrawMode == DrawModeValue_Meshes)
 		{
-			if (GizmoDrawMode == DrawModeValue_Meshes)
+			if (UViewAdjustedStaticMeshGizmoComponent* Component = AddMeshGizmoComponent(
+				TEXT("/Engine/InteractiveToolsFramework/Meshes/GizmoBoxHandle"),
+				UE::GizmoRenderingUtil::GetDefaultAxisColor(ElementAxis),
+				UE::GizmoUtil::GetRotatedBasisTransform(
+					FTransform(FQuat::Identity, FVector(130, 0, 0), FVector(0.8)),
+					ElementAxis),
+				EMirror::WhenCombined))
 			{
-				if (UViewAdjustedStaticMeshGizmoComponent* Component = AddMeshGizmoComponent(
-					TEXT("/Engine/InteractiveToolsFramework/Meshes/GizmoBoxHandle"),
+				Component->SetRenderVisibilityFunction(ShouldAxisBeVisible);
+
+				// Also try to add a full handle to use when we're not using a combined gizmo.
+				FullHandleOut = AddMeshGizmoComponent(
+					TEXT("/Engine/InteractiveToolsFramework/Meshes/GizmoBoxArrowHandle"),
 					UE::GizmoRenderingUtil::GetDefaultAxisColor(ElementAxis),
 					UE::GizmoUtil::GetRotatedBasisTransform(
-						FTransform(FQuat::Identity, FVector(130, 0, 0), FVector(0.8)),
+						FTransform(FQuat::Identity, FVector::ZeroVector, FVector::One()),
 						ElementAxis),
-					EMirror::WhenCombined))
+					EMirror::WhenCombined);
+				if (FullHandleOut)
 				{
-					Component->SetRenderVisibilityFunction(ShouldAxisBeVisible);
-
-					// Also try to add a full handle to use when we're not using a combined gizmo.
-					FullHandleOut = AddMeshGizmoComponent(
-						TEXT("/Engine/InteractiveToolsFramework/Meshes/GizmoBoxArrowHandle"),
-						UE::GizmoRenderingUtil::GetDefaultAxisColor(ElementAxis),
-						UE::GizmoUtil::GetRotatedBasisTransform(
-							FTransform(FQuat::Identity, FVector::ZeroVector, FVector::One()),
-							ElementAxis),
-						EMirror::WhenCombined);
-					if (FullHandleOut)
-					{
-						FullHandleOut->SetVisibility(false);
-					}
-
-					return Component;
+					FullHandleOut->SetVisibility(false);
 				}
+
+				return Component;
 			}
 		}
 
@@ -1173,6 +1170,11 @@ void UCombinedTransformGizmo::ApplyGizmoActiveMode()
 {
 	using namespace CombinedTransformGizmoLocals;
 
+	if (!ensure(IsValid(GizmoActor)))
+	{
+		return;
+	}
+
 	auto SetSubGizmoTypeVisibility = [this](TArray<FSubGizmoInfo>& GizmoInfos, bool bVisible)
 	{
 		for (FSubGizmoInfo& GizmoInfo : GizmoInfos)
