InteractiveToolsFramework: Fix static analysis warning in CombinedTransformGizmo.cpp.

diff --git a/Engine/Source/Runtime/InteractiveToolsFramework/Private/BaseGizmos/CombinedTransformGizmo.cpp b/Engine/Source/Runtime/InteractiveToolsFramework/Private/BaseGizmos/CombinedTransformGizmo.cpp
--- a/Engine/Source/Runtime/InteractiveToolsFramework/Private/BaseGizmos/CombinedTransformGizmo.cpp
+++ b/Engine/Source/Runtime/InteractiveToolsFramework/Private/BaseGizmos/CombinedTransformGizmo.cpp
@@ -1170,10 +1170,7 @@ void UCombinedTransformGizmo::ApplyGizmoActiveMode()
 {
 	using namespace CombinedTransformGizmoLocals;
 
-	if (!ensure(IsValid(GizmoActor)))
-	{
-		return;
-	}
+	ON_SCOPE_EXIT{ PreviousActiveGizmoMode = ActiveGizmoMode; };
 
 	auto SetSubGizmoTypeVisibility = [this](TArray<FSubGizmoInfo>& GizmoInfos, bool bVisible)
 	{
@@ -1199,6 +1196,12 @@ void UCombinedTransformGizmo::ApplyGizmoActiveMode()
 	SetSubGizmoTypeVisibility(RotationSubGizmos, bShouldShowRotation);
 	SetSubGizmoTypeVisibility(UniformScaleSubGizmos, bShouldShowUniformScale);
 
+	// The rest of the modifications dereference GizmoActor, so go ahead and do a safety check now.
+	if (!ensure(IsValid(GizmoActor)))
+	{
+		return;
+	}
+
 	if (ActiveGizmoMode == EToolContextTransformGizmoMode::Combined || ActiveGizmoMode == EToolContextTransformGizmoMode::Scale)
 	{
 		// The scale handles look different in different modes, so swap them if necessary
@@ -1270,7 +1273,7 @@ void UCombinedTransformGizmo::ApplyGizmoActiveMode()
 			|| (ActiveGizmoMode == EToolContextTransformGizmoMode::Combined && UniformScaleSubGizmos.Num() == 0));
 	}
 
-	if (GizmoActor && GizmoDrawMode == DrawModeValue_Meshes)
+	if (GizmoDrawMode == DrawModeValue_Meshes)
 	{
 		// Many components mirror in combined mode but not separate mode, mostly because it is weird for them
 		//  to not mirror when the rotation components do.
@@ -1284,8 +1287,6 @@ void UCombinedTransformGizmo::ApplyGizmoActiveMode()
 			}
 		}
 	}
-
-	PreviousActiveGizmoMode = ActiveGizmoMode;
 }
 
 
