[InteractiveToolsFramework]
* Fixed compile error caused by missing include

diff --git a/Engine/Source/Runtime/InteractiveToolsFramework/Private/BaseGizmos/ViewAdjustedStaticMeshGizmoComponent.cpp b/Engine/Source/Runtime/InteractiveToolsFramework/Private/BaseGizmos/ViewAdjustedStaticMeshGizmoComponent.cpp
--- a/Engine/Source/Runtime/InteractiveToolsFramework/Private/BaseGizmos/ViewAdjustedStaticMeshGizmoComponent.cpp
+++ b/Engine/Source/Runtime/InteractiveToolsFramework/Private/BaseGizmos/ViewAdjustedStaticMeshGizmoComponent.cpp
@@ -6,6 +6,7 @@
 #include "BaseGizmos/CombinedTransformGizmo.h"
 #include "BaseGizmos/ViewBasedTransformAdjusters.h"
 #include "BaseGizmos/GizmoRenderingUtil.h"
+#include "MaterialShared.h"
 #include "Materials/MaterialRenderProxy.h"
 #include "SceneInterface.h"
 #include "StaticMeshResources.h"
