InteractiveToolsFramework: Refactor UBrushAdjusterInputBehavior to support any brush tool.

Also add behavior to vertex and dynamic sculpt tools (which was not possible before this refactor).

diff --git a/Engine/Plugins/Experimental/MeshModelingToolsetExp/Source/MeshModelingToolsExp/Private/DynamicMeshSculptTool.cpp b/Engine/Plugins/Experimental/MeshModelingToolsetExp/Source/MeshModelingToolsExp/Private/DynamicMeshSculptTool.cpp
--- a/Engine/Plugins/Experimental/MeshModelingToolsetExp/Source/MeshModelingToolsExp/Private/DynamicMeshSculptTool.cpp
+++ b/Engine/Plugins/Experimental/MeshModelingToolsetExp/Source/MeshModelingToolsExp/Private/DynamicMeshSculptTool.cpp
@@ -32,6 +32,7 @@
 #include "Sculpting/KelvinletBrushOp.h"
 
 #include "InteractiveGizmoManager.h"
+#include "BaseBehaviors/TwoAxisPropertyEditBehavior.h"
 #include "BaseGizmos/GizmoComponents.h"
 #include "BaseGizmos/TransformGizmoUtil.h"
 #include "UObject/ObjectMacros.h"
@@ -302,6 +303,73 @@ void UDynamicMeshSculptTool::Setup()
 	}
 
 	UpdateBrushType(SculptProperties->PrimaryBrushType);
+
+	// add input behavior to click-drag while holding hotkey to adjust brush size and strength
+    {
+    	BrushEditBehavior = NewObject<ULocalTwoAxisPropertyEditInputBehavior>(this);
+		
+    	BrushEditBehavior->HorizontalProperty.GetValueFunc = [this]()
+    	{
+    		if (BrushProperties->BrushSize.SizeType == EBrushToolSizeType::Adaptive)
+    		{
+    			return BrushProperties->BrushSize.AdaptiveSize;
+    		}
+            else
+            {
+	            return BrushProperties->BrushSize.WorldRadius;
+            }
+    	};
+    	BrushEditBehavior->HorizontalProperty.SetValueFunc = [this](float NewValue)
+    	{
+    		if (BrushProperties->BrushSize.SizeType == EBrushToolSizeType::Adaptive)
+    		{
+    			BrushProperties->BrushSize.AdaptiveSize = FMath::Clamp(NewValue, 0.f, 1.f);
+    		}
+    		else
+    		{
+    			BrushProperties->BrushSize.WorldRadius = FMath::Max(NewValue, 0.01f);
+    		}
+    		
+#if WITH_EDITOR
+    		FPropertyChangedEvent PropertyChangedEvent(UDynamicMeshBrushProperties::StaticClass()->FindPropertyByName(GET_MEMBER_NAME_CHECKED(UDynamicMeshBrushProperties, BrushSize)));
+    		BrushProperties->PostEditChangeProperty(PropertyChangedEvent);
+#endif
+    	};
+		BrushEditBehavior->HorizontalProperty.MutateDeltaFunc = [this](float Delta)
+		{
+			// Scale delta if brush size is in world units.
+			return Delta * (BrushProperties->BrushSize.SizeType == EBrushToolSizeType::World ? (CameraState.Position - LastBrushPosWorld).Length() : 1.f);
+		};
+		BrushEditBehavior->HorizontalProperty.Name = LOCTEXT("BrushRadius", "Radius");
+		BrushEditBehavior->HorizontalProperty.bEnabled = true;
+		
+    	BrushEditBehavior->VerticalProperty.GetValueFunc = [this]()
+    	{
+    		return SculptProperties->PrimaryBrushSpeed;
+    	};
+    	BrushEditBehavior->VerticalProperty.SetValueFunc = [this](float NewValue)
+    	{
+    		SculptProperties->PrimaryBrushSpeed = FMath::Clamp(NewValue, 0.f, 1.f);
+
+#if WITH_EDITOR
+    		FPropertyChangedEvent PropertyChangedEvent(UDynamicMeshBrushSculptProperties::StaticClass()->FindPropertyByName(GET_MEMBER_NAME_CHECKED(UDynamicMeshBrushSculptProperties, PrimaryBrushSpeed)));
+    		BrushProperties->PostEditChangeProperty(PropertyChangedEvent);
+#endif
+    	};
+		BrushEditBehavior->VerticalProperty.Name = LOCTEXT("BrushStrength", "Strength");
+		BrushEditBehavior->VerticalProperty.bEnabled = true;
+
+
+		
+    	BrushEditBehavior->OnDragUpdated.AddWeakLambda(this, [this]()
+    	{
+    		CalculateBrushRadius();
+    		NotifyOfPropertyChangeByTool(BrushProperties);
+    	});
+    	
+    	BrushEditBehavior->Initialize();
+    	AddInputBehavior(BrushEditBehavior.Get());
+    }
 }
 
 void UDynamicMeshSculptTool::Shutdown(EToolShutdownType ShutdownType)
@@ -1528,6 +1596,22 @@ void UDynamicMeshSculptTool::Render(IToolsContextRenderAPI* RenderAPI)
 	{
 		double FixedDimScale = ToolSceneQueriesUtil::CalculateDimensionFromVisualAngleD(RenderCameraState, LastBrushPosWorld, 1.5f);
 		BrushIndicatorMaterial->SetScalarParameterValue(TEXT("FalloffWidth"), FixedDimScale);
+
+		if (BrushEditBehavior.IsValid())
+		{
+			FLinearColor IndicatorColor;
+		
+			if (BrushEditBehavior->IsEditing())
+			{
+				IndicatorColor = FLinearColor::White;
+			}
+			else
+			{
+				BrushIndicatorMaterial->GetVectorParameterDefaultValue(TEXT("IndicatorColor"), IndicatorColor);
+			}
+			
+			BrushIndicatorMaterial->SetVectorParameterValue(TEXT("IndicatorColor"), IndicatorColor);
+		}
 	}
 
 	if (SculptProperties->PrimaryBrushType == EDynamicMeshSculptBrushType::FixedPlane)
@@ -1541,6 +1625,15 @@ void UDynamicMeshSculptTool::Render(IToolsContextRenderAPI* RenderAPI)
 	}
 }
 
+void UDynamicMeshSculptTool::DrawHUD(FCanvas* Canvas, IToolsContextRenderAPI* RenderAPI)
+{
+	Super::DrawHUD(Canvas, RenderAPI);
+	if (BrushEditBehavior.IsValid())
+	{
+		BrushEditBehavior->DrawHUD(Canvas, RenderAPI);	
+	}
+}
+
 void UDynamicMeshSculptTool::OnTick(float DeltaTime)
 {
 	TRACE_CPUPROFILER_EVENT_SCOPE(DynamicMeshSculptTool_OnTick);
diff --git a/Engine/Plugins/Experimental/MeshModelingToolsetExp/Source/MeshModelingToolsExp/Private/MeshVertexSculptTool.cpp b/Engine/Plugins/Experimental/MeshModelingToolsetExp/Source/MeshModelingToolsExp/Private/MeshVertexSculptTool.cpp
--- a/Engine/Plugins/Experimental/MeshModelingToolsetExp/Source/MeshModelingToolsExp/Private/MeshVertexSculptTool.cpp
+++ b/Engine/Plugins/Experimental/MeshModelingToolsetExp/Source/MeshModelingToolsExp/Private/MeshVertexSculptTool.cpp
@@ -21,6 +21,7 @@
 
 #include "BaseGizmos/BrushStampIndicator.h"
 #include "PreviewMesh.h"
+#include "BaseBehaviors/TwoAxisPropertyEditBehavior.h"
 #include "Generators/RectangleMeshGenerator.h"
 
 #include "Changes/MeshVertexChange.h"
@@ -1283,6 +1284,16 @@ void UMeshVertexSculptTool::Render(IToolsContextRenderAPI* RenderAPI)
 }
 
 
+void UMeshVertexSculptTool::DrawHUD(FCanvas* Canvas, IToolsContextRenderAPI* RenderAPI)
+{
+	Super::DrawHUD(Canvas, RenderAPI);
+	if (BrushEditBehavior.IsValid())
+	{
+		BrushEditBehavior->DrawHUD(Canvas, RenderAPI);	
+	}
+}
+
+
 void UMeshVertexSculptTool::OnTick(float DeltaTime)
 {
 	UMeshSculptToolBase::OnTick(DeltaTime);
@@ -1696,6 +1707,27 @@ void UMeshVertexSculptTool::UpdateBrushType(int32 BrushID)
 
 	SetActivePrimaryBrushType(BrushID);
 
+	if (BrushEditBehavior.IsValid())
+	{
+		// todo: Handle Kelvinlet brush props better. At the moment we are just disabling strength editing for Kelvinlet brush ops.
+		auto PropertySetSupportsStrength = [this]()
+		{
+			return PrimaryBrushOp->PropertySet.IsValid() && Cast<UBaseKelvinletBrushOpProps>(PrimaryBrushOp->PropertySet.Get()) == nullptr;
+		};
+	
+		if (PropertySetSupportsStrength())
+		{
+			BrushEditBehavior->VerticalProperty.Name = LOCTEXT("BrushStrength", "Strength");
+            BrushEditBehavior->VerticalProperty.GetValueFunc = [this](){ return PrimaryBrushOp->PropertySet->GetStrength(); };
+            BrushEditBehavior->VerticalProperty.SetValueFunc = [this](float NewValue){ PrimaryBrushOp->PropertySet->SetStrength(FMath::Clamp(NewValue, 0.f, 1.f)); };
+            BrushEditBehavior->VerticalProperty.bEnabled = true;
+		}
+		else
+		{
+			BrushEditBehavior->VerticalProperty.bEnabled = false;
+		}
+	}
+	
 	if (ensure(SculptProperties))
 	{
 		SculptProperties->bCanFreezeTarget = 
diff --git a/Engine/Plugins/Experimental/MeshModelingToolsetExp/Source/MeshModelingToolsExp/Private/Sculpting/MeshSculptToolBase.cpp b/Engine/Plugins/Experimental/MeshModelingToolsetExp/Source/MeshModelingToolsExp/Private/Sculpting/MeshSculptToolBase.cpp
--- a/Engine/Plugins/Experimental/MeshModelingToolsetExp/Source/MeshModelingToolsExp/Private/Sculpting/MeshSculptToolBase.cpp
+++ b/Engine/Plugins/Experimental/MeshModelingToolsetExp/Source/MeshModelingToolsExp/Private/Sculpting/MeshSculptToolBase.cpp
@@ -16,6 +16,7 @@
 #include "Generators/SphereGenerator.h"
 
 #include "ModelingToolTargetUtil.h"
+#include "BaseBehaviors/TwoAxisPropertyEditBehavior.h"
 
 #include UE_INLINE_GENERATED_CPP_BY_NAME(MeshSculptToolBase)
 
@@ -152,6 +153,55 @@ void UMeshSculptToolBase::Setup()
 		[this](bool bOn) { UpdateTwoSidedSetting(bOn); });
 	ViewProperties->WatchProperty(ViewProperties->Image,
 		[this](UTexture2D* NewImage) { UpdateImageSetting(NewImage); });
+	
+	// add input behavior to click-drag while holding hotkey to adjust brush size and strength
+    {
+        BrushEditBehavior = NewObject<ULocalTwoAxisPropertyEditInputBehavior>(this);
+		
+        BrushEditBehavior->HorizontalProperty.GetValueFunc = [this]()
+        {
+        	if (BrushProperties->BrushSize.SizeType == EBrushToolSizeType::Adaptive)
+        	{
+        		return BrushProperties->BrushSize.AdaptiveSize;
+        	}
+            else
+            {
+                return BrushProperties->BrushSize.WorldRadius;
+            }
+        };
+        BrushEditBehavior->HorizontalProperty.SetValueFunc = [this](float NewValue)
+        {
+        	if (BrushProperties->BrushSize.SizeType == EBrushToolSizeType::Adaptive)
+        	{
+        		BrushProperties->BrushSize.AdaptiveSize = FMath::Clamp(NewValue, 0.f, 1.f);
+        	}
+        	else
+        	{
+        		BrushProperties->BrushSize.WorldRadius = FMath::Max(NewValue, 0.01f);
+        	}
+        	
+#if WITH_EDITOR
+        	FPropertyChangedEvent PropertyChangedEvent(USculptBrushProperties::StaticClass()->FindPropertyByName(GET_MEMBER_NAME_CHECKED(USculptBrushProperties, BrushSize)));
+        	BrushProperties->PostEditChangeProperty(PropertyChangedEvent);
+#endif
+        };
+		BrushEditBehavior->HorizontalProperty.MutateDeltaFunc = [this](float Delta)
+		{
+			// Scale delta if brush size is in world units.
+			return Delta * (BrushProperties->BrushSize.SizeType == EBrushToolSizeType::World ? (CameraState.Position - LastBrushFrameWorld.Origin).Length() : 1.f);
+		};
+		BrushEditBehavior->HorizontalProperty.Name = LOCTEXT("BrushRadius", "Radius");
+		BrushEditBehavior->HorizontalProperty.bEnabled = true;
+    	
+        BrushEditBehavior->OnDragUpdated.AddWeakLambda(this, [this]()
+        {
+        	CalculateBrushRadius();
+        	NotifyOfPropertyChangeByTool(BrushProperties);
+        });
+        
+        BrushEditBehavior->Initialize();
+        AddInputBehavior(BrushEditBehavior.Get());
+    }
 }
 
 
diff --git a/Engine/Plugins/Experimental/MeshModelingToolsetExp/Source/MeshModelingToolsExp/Public/DynamicMeshSculptTool.h b/Engine/Plugins/Experimental/MeshModelingToolsetExp/Source/MeshModelingToolsExp/Public/DynamicMeshSculptTool.h
--- a/Engine/Plugins/Experimental/MeshModelingToolsetExp/Source/MeshModelingToolsExp/Public/DynamicMeshSculptTool.h
+++ b/Engine/Plugins/Experimental/MeshModelingToolsetExp/Source/MeshModelingToolsExp/Public/DynamicMeshSculptTool.h
@@ -26,6 +26,7 @@
 
 class UCombinedTransformGizmo;
 class UTransformProxy;
+class ULocalTwoAxisPropertyEditInputBehavior;
 class UMaterialInstanceDynamic;
 
 class FMeshVertexChangeBuilder;
@@ -274,6 +275,7 @@ public:
 
 	virtual void OnTick(float DeltaTime) override;
 	virtual void Render(IToolsContextRenderAPI* RenderAPI) override;
+	virtual void DrawHUD(FCanvas* Canvas, IToolsContextRenderAPI* RenderAPI) override;
 
 	virtual bool HasCancel() const override { return true; }
 	virtual bool HasAccept() const override { return true; }
@@ -342,6 +344,8 @@ private:
 	UWorld* TargetWorld;		// required to spawn UPreviewMesh/etc
 	FViewCameraState CameraState;
 
+	TWeakObjectPtr<ULocalTwoAxisPropertyEditInputBehavior> BrushEditBehavior;
+
 	UPROPERTY()
 	TObjectPtr<UBrushStampIndicator> BrushIndicator;
 
diff --git a/Engine/Plugins/Experimental/MeshModelingToolsetExp/Source/MeshModelingToolsExp/Public/MeshVertexSculptTool.h b/Engine/Plugins/Experimental/MeshModelingToolsetExp/Source/MeshModelingToolsExp/Public/MeshVertexSculptTool.h
--- a/Engine/Plugins/Experimental/MeshModelingToolsetExp/Source/MeshModelingToolsExp/Public/MeshVertexSculptTool.h
+++ b/Engine/Plugins/Experimental/MeshModelingToolsetExp/Source/MeshModelingToolsExp/Public/MeshVertexSculptTool.h
@@ -217,6 +217,7 @@ public:
 	virtual void Setup() override;
 	virtual void Shutdown(EToolShutdownType ShutdownType) override;
 	virtual void Render(IToolsContextRenderAPI* RenderAPI) override;
+	virtual void DrawHUD(FCanvas* Canvas, IToolsContextRenderAPI* RenderAPI) override;
 
 	virtual void OnTick(float DeltaTime) override;
 
diff --git a/Engine/Plugins/Experimental/MeshModelingToolsetExp/Source/MeshModelingToolsExp/Public/Sculpting/MeshSculptToolBase.h b/Engine/Plugins/Experimental/MeshModelingToolsetExp/Source/MeshModelingToolsExp/Public/Sculpting/MeshSculptToolBase.h
--- a/Engine/Plugins/Experimental/MeshModelingToolsetExp/Source/MeshModelingToolsExp/Public/Sculpting/MeshSculptToolBase.h
+++ b/Engine/Plugins/Experimental/MeshModelingToolsetExp/Source/MeshModelingToolsExp/Public/Sculpting/MeshSculptToolBase.h
@@ -14,6 +14,7 @@
 #include "MeshSculptToolBase.generated.h"
 
 
+class ULocalTwoAxisPropertyEditInputBehavior;
 class UMaterialInstanceDynamic;
 class UCombinedTransformGizmo;
 class UTransformProxy;
@@ -567,6 +568,8 @@ protected:
 	UPROPERTY()
 	TObjectPtr<UPreviewMesh> BrushIndicatorMesh;
 
+	
+	TWeakObjectPtr<ULocalTwoAxisPropertyEditInputBehavior> BrushEditBehavior;
 
 
 	//
diff --git a/Engine/Source/Runtime/InteractiveToolsFramework/Private/BaseBehaviors/TwoAxisPropertyEditBehavior.cpp b/Engine/Source/Runtime/InteractiveToolsFramework/Private/BaseBehaviors/TwoAxisPropertyEditBehavior.cpp
new file mode 100644
--- /dev/null
+++ b/Engine/Source/Runtime/InteractiveToolsFramework/Private/BaseBehaviors/TwoAxisPropertyEditBehavior.cpp
@@ -0,0 +1,228 @@
+// Copyright Epic Games, Inc. All Rights Reserved.
+
+#include "BaseBehaviors/TwoAxisPropertyEditBehavior.h"
+
+#include "CanvasItem.h"
+#include "CanvasTypes.h"
+#include "Engine/Engine.h"
+#include "Engine/HitResult.h"
+#include "GenericPlatform/GenericPlatformApplicationMisc.h"
+#include "ToolContextInterfaces.h"
+
+#include UE_INLINE_GENERATED_CPP_BY_NAME(TwoAxisPropertyEditBehavior)
+
+#define LOCTEXT_NAMESPACE "UTwoAxisPropertyEditInputBehavior"
+
+void UTwoAxisPropertyEditInputBehavior::Initialize(ITwoAxisPropertyEditBehaviorTarget* InTarget)
+{
+	check(InTarget);
+	Target = InTarget;
+	
+	ResetDragState();
+}
+
+void UTwoAxisPropertyEditInputBehavior::ResetOrigin(FVector2D InScreenPosition, bool bHorizontalAdjust, bool bResetStartOrigin)
+{
+	using IPropertyInterface = ITwoAxisPropertyEditBehaviorTarget::IPropertyInterface;
+
+	if (bResetStartOrigin)
+	{
+		State.StartOrigin = InScreenPosition;
+	}
+	
+	State.bAdjustingHorizontally = bHorizontalAdjust;
+	State.CurrentOrigin = InScreenPosition;
+	
+	if (IPropertyInterface* Property = State.bAdjustingHorizontally ? Target->GetHorizontalProperty() : Target->GetVerticalProperty(); Property && Property->IsEnabled())
+	{
+		State.StartValue = Property->GetValue();
+	}
+}
+
+void UTwoAxisPropertyEditInputBehavior::DrawHUD(FCanvas* Canvas, IToolsContextRenderAPI* RenderAPI) const
+{
+	using IPropertyInterface = ITwoAxisPropertyEditBehaviorTarget::IPropertyInterface;
+	
+	if (!IsEditing())
+	{
+		return;
+	}
+
+	FText BrushAdjustmentMessage;
+	if (IPropertyInterface* Property = State.bAdjustingHorizontally ? Target->GetHorizontalProperty() : Target->GetVerticalProperty(); Property && Property->IsEnabled())
+	{
+		BrushAdjustmentMessage = FText::Format(LOCTEXT("BrushAdjustmentMessage", "{0}: {1}"), Property->GetName(), FText::AsNumber(Property->GetValue()));
+	}
+	else
+	{
+		return;
+	}
+
+	FCanvasTextItem TextItem(State.StartOrigin, BrushAdjustmentMessage, GEngine->GetMediumFont(), FLinearColor::White);
+	TextItem.EnableShadow(FLinearColor::Black);
+	Canvas->DrawItem(TextItem);
+}
+
+void UTwoAxisPropertyEditInputBehavior::OnDragStart(FVector2D InScreenPosition)
+{
+	bInDrag = true;
+	constexpr bool bHorizontalAdjust = true;
+	constexpr bool bResetStartOrigin = true;
+	ResetOrigin(InScreenPosition, bHorizontalAdjust, bResetStartOrigin);
+}
+
+void UTwoAxisPropertyEditInputBehavior::OnDragUpdate(FVector2D InScreenPosition)
+{
+	using IPropertyInterface = ITwoAxisPropertyEditBehaviorTarget::IPropertyInterface;
+	
+	// Calculate screen space cursor delta relative to current origin
+	const float HorizontalDelta = InScreenPosition.X - State.CurrentOrigin.X;
+	const float VerticalDelta = InScreenPosition.Y - State.CurrentOrigin.Y;
+	const float HorizDeltaMag = FMath::Abs(HorizontalDelta);
+	const float VertDeltaMag = FMath::Abs(VerticalDelta);
+	
+	// Apply directional adjustments
+	if (IPropertyInterface* Property = State.bAdjustingHorizontally ? Target->GetHorizontalProperty() : Target->GetVerticalProperty(); Property && Property->IsEnabled())
+	{
+		const float DPIScale = FGenericPlatformApplicationMisc::GetDPIScaleFactorAtPoint(InScreenPosition.X, InScreenPosition.Y);
+		const float Delta = State.bAdjustingHorizontally ? HorizontalDelta : VerticalDelta;
+		
+		const float NewValue = State.StartValue + (State.bAdjustingHorizontally ? 1.f : -1.f) * Property->MutateDelta(Delta * Property->GetEditRate() * DPIScale);
+		Property->SetValue(NewValue);
+	}
+
+	// Determine if we need to swap adjustment axis
+	if (State.bAdjustingHorizontally && HorizDeltaMag < VertDeltaMag && Target->GetVerticalProperty()->IsEnabled())
+	{
+		// Switch to adjusting vertically and re-center current origin
+		ResetOrigin(InScreenPosition, false);
+	}
+	else if (!State.bAdjustingHorizontally && VertDeltaMag < HorizDeltaMag && Target->GetHorizontalProperty()->IsEnabled())
+	{
+		// Switch to adjusting horizontally and re-center current origin
+		ResetOrigin(InScreenPosition, true);
+	}
+	
+	Target->PostDragUpdated();
+}
+
+void UTwoAxisPropertyEditInputBehavior::OnDragEnd()
+{
+	bInDrag = false;
+	ResetDragState();
+}
+
+void UTwoAxisPropertyEditInputBehavior::ResetDragState()
+{
+	State = FDragState();
+}
+
+FInputCaptureRequest UTwoAxisPropertyEditInputBehavior::WantsCapture(const FInputDeviceState& InputState)
+{
+	switch (InputState.InputDevice)
+	{
+	case EInputDevices::Mouse: return WantsMouseCapture(InputState);
+	case EInputDevices::Keyboard: return WantsKeyboardCapture(InputState);
+	default: return FInputCaptureRequest::Ignore();
+	}
+}
+
+FInputCaptureUpdate UTwoAxisPropertyEditInputBehavior::BeginCapture(const FInputDeviceState& InputState, EInputCaptureSide eSide)
+{
+	switch (InputState.InputDevice)
+	{
+	case EInputDevices::Mouse: return BeginMouseCapture(InputState, eSide);
+	case EInputDevices::Keyboard: return BeginKeyboardCapture(InputState, eSide);
+	default: return FInputCaptureUpdate::Ignore();
+	}
+}
+
+FInputCaptureUpdate UTwoAxisPropertyEditInputBehavior::UpdateCapture(const FInputDeviceState& InputState, const FInputCaptureData& CaptureData)
+{
+	switch (InputState.InputDevice)
+	{
+	case EInputDevices::Mouse: return UpdateMouseCapture(InputState, CaptureData);
+	case EInputDevices::Keyboard: return UpdateKeyboardCapture(InputState, CaptureData);
+	default: return FInputCaptureUpdate::Ignore();
+	}
+}
+
+void UTwoAxisPropertyEditInputBehavior::ForceEndCapture(const FInputCaptureData& CaptureData)
+{
+	OnDragEnd();
+}
+
+FInputCaptureRequest UTwoAxisPropertyEditInputBehavior::WantsMouseCapture(const FInputDeviceState& InputState)
+{
+	if (IsMousePressed(InputState))
+	{
+		return FInputCaptureRequest::Begin(this, EInputCaptureSide::Any, 0.0);
+	}
+	else
+	{
+		return FInputCaptureRequest::Ignore();
+	}
+}
+
+FInputCaptureUpdate UTwoAxisPropertyEditInputBehavior::BeginMouseCapture(const FInputDeviceState& InputState, EInputCaptureSide eSide)
+{
+	OnDragStart(InputState.Mouse.Position2D);
+	return FInputCaptureUpdate::Begin(this, EInputCaptureSide::Any);
+}
+
+FInputCaptureUpdate UTwoAxisPropertyEditInputBehavior::UpdateMouseCapture(const FInputDeviceState& InputState, const FInputCaptureData& CaptureData)
+{
+	if (IsMousePressed(InputState))
+	{
+		OnDragUpdate(InputState.Mouse.Position2D);
+		return FInputCaptureUpdate::Continue();
+	}
+	else
+	{
+		OnDragEnd();
+		return FInputCaptureUpdate::End();
+	}
+}
+
+bool UTwoAxisPropertyEditInputBehavior::IsMousePressed(const FInputDeviceState& InputState) const
+{
+	return InputState.Mouse.Left.bDown && bKeyPressed;
+}
+
+FInputCaptureRequest UTwoAxisPropertyEditInputBehavior::WantsKeyboardCapture(const FInputDeviceState& InputState)
+{
+	if (IsKeyboardPressed(InputState))
+	{
+		return FInputCaptureRequest::Begin(this, EInputCaptureSide::Any, 0.0);
+	}
+	else
+	{
+		return FInputCaptureRequest::Ignore();
+	}
+}
+
+FInputCaptureUpdate UTwoAxisPropertyEditInputBehavior::BeginKeyboardCapture(const FInputDeviceState& InputState, EInputCaptureSide eSide)
+{
+	bKeyPressed = true;
+	return FInputCaptureUpdate::Begin(this, EInputCaptureSide::Any);
+}
+
+FInputCaptureUpdate UTwoAxisPropertyEditInputBehavior::UpdateKeyboardCapture(const FInputDeviceState& InputState, const FInputCaptureData& CaptureData)
+{
+	if (IsKeyboardPressed(InputState))
+    {
+    	return FInputCaptureUpdate::Continue();
+    }
+    else
+    {
+    	bKeyPressed = false;
+    	return FInputCaptureUpdate::End();
+    }
+}
+
+bool UTwoAxisPropertyEditInputBehavior::IsKeyboardPressed(const FInputDeviceState& InputState) const
+{
+	return InputState.Keyboard.ActiveKey.Button == Target->GetCaptureKey() && InputState.Keyboard.ActiveKey.bDown;
+}
+
+#undef LOCTEXT_NAMESPACE
\ No newline at end of file
diff --git a/Engine/Source/Runtime/InteractiveToolsFramework/Private/BaseTools/BaseBrushTool.cpp b/Engine/Source/Runtime/InteractiveToolsFramework/Private/BaseTools/BaseBrushTool.cpp
--- a/Engine/Source/Runtime/InteractiveToolsFramework/Private/BaseTools/BaseBrushTool.cpp
+++ b/Engine/Source/Runtime/InteractiveToolsFramework/Private/BaseTools/BaseBrushTool.cpp
@@ -9,7 +9,7 @@
 #include "InteractiveGizmoManager.h"
 #include "InteractiveTool.h"
 #include "GenericPlatform/GenericPlatformApplicationMisc.h"
-#include "BaseBehaviors/ClickDragBehavior.h"
+#include "BaseBehaviors/TwoAxisPropertyEditBehavior.h"
 #include "BaseGizmos/BrushStampIndicator.h"
 
 #include UE_INLINE_GENERATED_CPP_BY_NAME(BaseBrushTool)
@@ -225,39 +225,71 @@ void UBaseBrushTool::Setup()
 	// add input behavior to click-drag while holding hotkey to adjust brush size and strength
 	if (SupportsBrushAdjustmentInput())
 	{
-		BrushAdjusterBehavior = NewObject<UBrushAdjusterInputBehavior>(this);
-		BrushAdjusterBehavior->Initialize(this);
-		AddInputBehavior(BrushAdjusterBehavior.Get());
-
-		ULocalClickDragInputBehavior* BrushAdjusterClickDragBehavior = NewObject<ULocalClickDragInputBehavior>(this);
-		BrushAdjusterClickDragBehavior->Initialize();
-		BrushAdjusterClickDragBehavior->SetDefaultPriority(FInputCapturePriority().MakeHigher());
-		BrushAdjusterClickDragBehavior->ModifierCheckFunc = [this](const FInputDeviceState&)
-		{
-			return BrushAdjusterBehavior.IsValid() ? BrushAdjusterBehavior->IsBrushBeingAdjusted() : false;
-		};
-		BrushAdjusterClickDragBehavior->CanBeginClickDragFunc = [](const FInputDeviceRay&)
-		{
-			// fake screen hit
-			return FInputRayHit(0.f);
-		};
-		BrushAdjusterClickDragBehavior->OnClickPressFunc = [this](const FInputDeviceRay& PressPos)
+		BrushEditBehavior = NewObject<ULocalTwoAxisPropertyEditInputBehavior>(this);
+		
+		BrushEditBehavior->HorizontalProperty.GetValueFunc = [this]()
 		{
-			if (BrushAdjusterBehavior.IsValid())
+			if (BrushProperties->bSpecifyRadius)
 			{
-				BrushAdjusterBehavior->OnDragStart(PressPos.ScreenPosition);
+				return BrushProperties->BrushRadius;
+			}
+			else
+			{
+				return BrushProperties->BrushSize;
 			}
 		};
-		BrushAdjusterClickDragBehavior->OnClickDragFunc = [this](const FInputDeviceRay& DragPos)
+		BrushEditBehavior->HorizontalProperty.SetValueFunc = [this](float NewValue)
 		{
-			if (BrushAdjusterBehavior.IsValid())
+			if (BrushProperties->bSpecifyRadius)
+			{
+				BrushProperties->BrushRadius = FMath::Max(NewValue, 0.01f);
+#if WITH_EDITOR
+				FPropertyChangedEvent PropertyChangedEvent(UBrushBaseProperties::StaticClass()->FindPropertyByName(GET_MEMBER_NAME_CHECKED(UBrushBaseProperties, BrushRadius)));
+				BrushProperties->PostEditChangeProperty(PropertyChangedEvent);
+#endif
+			}
+			else
 			{
-				BrushAdjusterBehavior->OnDragUpdate(DragPos.ScreenPosition);
-				RecalculateBrushRadius();
-				NotifyOfPropertyChangeByTool(BrushProperties);
+				BrushProperties->BrushSize = FMath::Clamp(NewValue, 0.f, 1.f);
+#if WITH_EDITOR
+				FPropertyChangedEvent PropertyChangedEvent(UBrushBaseProperties::StaticClass()->FindPropertyByName(GET_MEMBER_NAME_CHECKED(UBrushBaseProperties, BrushSize)));
+				BrushProperties->PostEditChangeProperty(PropertyChangedEvent);
+#endif
 			}
 		};
-		AddInputBehavior(BrushAdjusterClickDragBehavior);
+		BrushEditBehavior->HorizontalProperty.MutateDeltaFunc = [this](float Delta)
+		{
+			// Scale delta if brush size is in world units.
+			return Delta * (BrushProperties->bSpecifyRadius ? LastBrushStamp.HitResult.Distance : 1.f);
+		};
+		BrushEditBehavior->HorizontalProperty.Name = LOCTEXT("BrushRadius", "Radius");
+		BrushEditBehavior->HorizontalProperty.EditRate = 0.001f;
+		BrushEditBehavior->HorizontalProperty.bEnabled = true;
+		
+		BrushEditBehavior->VerticalProperty.GetValueFunc = [this]()
+		{
+			return BrushProperties->BrushStrength;
+		};
+		BrushEditBehavior->VerticalProperty.SetValueFunc = [this](float NewValue)
+		{
+			BrushProperties->BrushStrength = FMath::Min(1.0f,FMath::Max(NewValue, 0.f));
+#if WITH_EDITOR
+			FPropertyChangedEvent PropertyChangedEvent(UBrushBaseProperties::StaticClass()->FindPropertyByName(GET_MEMBER_NAME_CHECKED(UBrushBaseProperties, BrushStrength)));
+			BrushProperties->PostEditChangeProperty(PropertyChangedEvent);
+#endif
+		};
+		BrushEditBehavior->VerticalProperty.Name = LOCTEXT("BrushStrength", "Strength");
+		BrushEditBehavior->VerticalProperty.EditRate = 0.005f;
+		BrushEditBehavior->VerticalProperty.bEnabled = true;
+		
+		BrushEditBehavior->OnDragUpdated.AddWeakLambda(this, [this]()
+		{
+			RecalculateBrushRadius();
+			NotifyOfPropertyChangeByTool(BrushProperties);
+		});
+		
+		BrushEditBehavior->Initialize();
+		AddInputBehavior(BrushEditBehavior.Get());
 	}
 }
 
@@ -476,16 +508,16 @@ void UBaseBrushTool::OnCancelDrag()
 
 bool UBaseBrushTool::OnUpdateHover(const FInputDeviceRay& DevicePos)
 {
-	if (BrushAdjusterBehavior.IsValid())
+	if (BrushEditBehavior.IsValid())
 	{
-		if (BrushAdjusterBehavior->IsBrushBeingAdjusted())
+		if (BrushEditBehavior->IsEditing())
 		{
 			return true;
 		}
 
 		// When not in adjustment mode, keep the brush & adjustment origin synchronized with
-		// the brush so that the initial BrushAdjuster HUD display tracks the brush stamp.
-		BrushAdjusterBehavior->OnDragStart(DevicePos.ScreenPosition);
+		// the brush so that the initial BrushEdit HUD display tracks the brush stamp.
+		BrushEditBehavior->ResetOrigin(DevicePos.ScreenPosition);
 	}
 	
 	FHitResult OutHit;
@@ -512,9 +544,9 @@ void UBaseBrushTool::Render(IToolsContextRenderAPI* RenderAPI)
 void UBaseBrushTool::DrawHUD(FCanvas* Canvas, IToolsContextRenderAPI* RenderAPI)
 {
 	Super::DrawHUD(Canvas, RenderAPI);
-	if (BrushAdjusterBehavior.IsValid())
+	if (BrushEditBehavior.IsValid())
 	{
-		BrushAdjusterBehavior->DrawHUD(Canvas, RenderAPI);	
+		BrushEditBehavior->DrawHUD(Canvas, RenderAPI);	
 	}
 }
 
@@ -534,9 +566,9 @@ void UBaseBrushTool::UpdateBrushStampIndicator()
 {
 	if (BrushStampIndicator)
 	{
-		if (BrushAdjusterBehavior.IsValid())
+		if (BrushEditBehavior.IsValid())
 		{
-			BrushStampIndicator->LineColor = BrushAdjusterBehavior->IsBrushBeingAdjusted() ? FLinearColor::White : FLinearColor::Green;	
+			BrushStampIndicator->LineColor = BrushEditBehavior->IsEditing() ? FLinearColor::White : FLinearColor::Green;	
 		}
 		
 		BrushStampIndicator->Update(BrushProperties->BrushRadius, LastBrushStamp.WorldPosition, LastBrushStamp.WorldNormal, BrushProperties->BrushFalloffAmount, BrushProperties->BrushStrength);
diff --git a/Engine/Source/Runtime/InteractiveToolsFramework/Public/BaseBehaviors/BehaviorTargetInterfaces.h b/Engine/Source/Runtime/InteractiveToolsFramework/Public/BaseBehaviors/BehaviorTargetInterfaces.h
--- a/Engine/Source/Runtime/InteractiveToolsFramework/Public/BaseBehaviors/BehaviorTargetInterfaces.h
+++ b/Engine/Source/Runtime/InteractiveToolsFramework/Public/BaseBehaviors/BehaviorTargetInterfaces.h
@@ -240,3 +240,36 @@ public:
 	 */
 	virtual void OnEndHover() = 0;
 };
+
+
+/**
+ * ITwoAxisPropertyEditBehaviorTarget allows tools to
+ * expose properties to drag-based property edit behaviors.
+ */
+class ITwoAxisPropertyEditBehaviorTarget
+{
+public:
+
+	/* IPropertyInterface encapsulates a virtual 'property' (which may or may not map to a single member or property) */
+	struct IPropertyInterface
+	{
+		virtual ~IPropertyInterface() {}
+		
+		virtual FText GetName() = 0;								// Get a name for this virtual property (for display purposes).
+		virtual float GetValue() = 0;
+		virtual void SetValue(float NewValue) = 0;
+		virtual float GetEditRate() { return 0.003f; }				// Scales how quickly this property is edited with respect to a fixed mouse movement.
+		virtual float MutateDelta(float Delta) { return Delta; }	// Provides a hook to further augment the delta that is about to be applied to this property.
+		virtual bool IsEnabled() { return false; }
+	};
+	
+public:
+	
+	virtual ~ITwoAxisPropertyEditBehaviorTarget() {}
+	
+	virtual IPropertyInterface* GetHorizontalProperty() = 0;
+	virtual IPropertyInterface* GetVerticalProperty() = 0;
+	
+	virtual FKey GetCaptureKey() { return EKeys::B; }
+	virtual void PostDragUpdated() {}
+};
\ No newline at end of file
diff --git a/Engine/Source/Runtime/InteractiveToolsFramework/Public/BaseBehaviors/TwoAxisPropertyEditBehavior.h b/Engine/Source/Runtime/InteractiveToolsFramework/Public/BaseBehaviors/TwoAxisPropertyEditBehavior.h
new file mode 100644
--- /dev/null
+++ b/Engine/Source/Runtime/InteractiveToolsFramework/Public/BaseBehaviors/TwoAxisPropertyEditBehavior.h
@@ -0,0 +1,145 @@
+// Copyright Epic Games, Inc. All Rights Reserved.
+
+#pragma once
+
+#include "InputBehavior.h"
+#include "BehaviorTargetInterfaces.h"
+
+#include "TwoAxisPropertyEditBehavior.generated.h"
+
+class FCanvas;
+
+class IToolsContextRenderAPI;
+
+/**
+ * A behavior that captures a keyboard hotkey to enter a property adjustment sub-mode while the key is pressed.
+ * In this sub-mode, click-dragging the mouse will begin updating the properties specified by the behavior target.
+ */
+UCLASS(MinimalAPI)
+class UTwoAxisPropertyEditInputBehavior : public UInputBehavior
+{
+	GENERATED_BODY()
+
+public:
+	
+	INTERACTIVETOOLSFRAMEWORK_API void Initialize(ITwoAxisPropertyEditBehaviorTarget* InTarget);
+
+	// Used to define the initial adjustment frame or to update the current reference point alternating between editing the 2 properties.
+	INTERACTIVETOOLSFRAMEWORK_API void ResetOrigin(FVector2D InScreenPosition, bool bHorizontalAdjust = true, bool bResetStartOrigin = false);
+	INTERACTIVETOOLSFRAMEWORK_API void DrawHUD(FCanvas* Canvas, IToolsContextRenderAPI* RenderAPI) const;
+	INTERACTIVETOOLSFRAMEWORK_API bool IsEditing() const { return bInDrag; }
+
+private:
+
+	void OnDragStart(FVector2D InScreenPosition);
+	void OnDragUpdate(FVector2D InScreenPosition);
+	void OnDragEnd();
+	void ResetDragState();
+
+	// UInputBehavior implementation
+	// Routes input to the device specific handler functions.
+	virtual EInputDevices GetSupportedDevices() override { return EInputDevices::Mouse | EInputDevices::Keyboard; }
+	virtual FInputCaptureRequest WantsCapture(const FInputDeviceState& InputState) override;
+	virtual FInputCaptureUpdate BeginCapture(const FInputDeviceState& InputState, EInputCaptureSide eSide) override;
+	virtual FInputCaptureUpdate UpdateCapture(const FInputDeviceState& InputState, const FInputCaptureData& CaptureData) override;
+	virtual void ForceEndCapture(const FInputCaptureData& CaptureData) override;
+
+	// Mouse input handling - In these functions we assume input is from mouse device.
+	// This code is responsible for invoking OnDrag* functions at the proper times.
+	FInputCaptureRequest WantsMouseCapture(const FInputDeviceState& InputState);
+	FInputCaptureUpdate BeginMouseCapture(const FInputDeviceState& InputState, EInputCaptureSide eSide);
+	FInputCaptureUpdate UpdateMouseCapture(const FInputDeviceState& InputState, const FInputCaptureData& CaptureData);
+	bool IsMousePressed(const FInputDeviceState& InputState) const;
+	
+	// Keyboard input handling - In these functions we assume input is from keyboard device.
+	// This code is responsible for keeping bKeyPressed in the proper state.
+	FInputCaptureRequest WantsKeyboardCapture(const FInputDeviceState& InputState);
+	FInputCaptureUpdate BeginKeyboardCapture(const FInputDeviceState& InputState, EInputCaptureSide eSide);
+	FInputCaptureUpdate UpdateKeyboardCapture(const FInputDeviceState& InputState, const FInputCaptureData& CaptureData);
+	bool IsKeyboardPressed(const FInputDeviceState& InputState) const;
+	
+private:
+
+	// These 2 bools reflect which input we are currently capturing.
+	bool bInDrag = false;
+	bool bKeyPressed = false;
+	
+	// Data that is only valid while in a drag (while bInDrag is true13).
+	struct FDragState
+	{
+		float StartValue = 0.f;				// Cached initial values of the property being actively edited.
+		FVector2D StartOrigin;				// The screen space coordinate of the brush when the drag started.
+		FVector2D CurrentOrigin;			// The screen space coordinate of the origin of the current adjustment. Reset when changing between horizontal and vertical adjustment.
+		bool bAdjustingHorizontally = true;
+	} State;								// Only valid if bInDrag is true
+	
+	ITwoAxisPropertyEditBehaviorTarget* Target = nullptr;
+};
+
+
+
+/**
+ * Variant of the base behavior which allows tools to use
+ * lambda functions instead of explicitly defining a behavior target.
+ */
+UCLASS(MinimalAPI)
+class ULocalTwoAxisPropertyEditInputBehavior : public UTwoAxisPropertyEditInputBehavior, public ITwoAxisPropertyEditBehaviorTarget
+{
+	GENERATED_BODY()
+
+public:
+	
+	/* Lambda implementation of IPropertyInterface */
+	struct FPropertyInterface : public IPropertyInterface
+	{
+		typedef TUniqueFunction<float(void)> FGetValueSignature;
+		typedef TUniqueFunction<void(float)> FSetValueSignature;
+		typedef TUniqueFunction<float(float)> FMutateDeltaSignature;
+
+		FGetValueSignature GetValueFunc;
+		FSetValueSignature SetValueFunc;
+		FMutateDeltaSignature MutateDeltaFunc = [](float Delta){ return Delta; };
+
+		FText Name;
+		float EditRate = 0.002f;
+		bool bEnabled = false;
+
+		bool IsValid() const
+		{
+			/* Enabled properties must have all mandatory fields set, disabled properties are always valid. */
+			if (bEnabled) { return !Name.IsEmpty() && GetValueFunc && SetValueFunc; }
+			return true;
+		}
+		
+	private:
+		
+		virtual FText GetName() override { return Name; }
+		virtual float GetValue() override { return GetValueFunc(); }
+		virtual void SetValue(float NewValue) override { SetValueFunc(NewValue); }
+		virtual float GetEditRate() override { return EditRate; }
+		virtual float MutateDelta(float Delta) override { return MutateDeltaFunc(Delta); }
+		virtual bool IsEnabled() override { return bEnabled; }
+	};
+	
+public:
+	
+	using UTwoAxisPropertyEditInputBehavior::Initialize;
+	void Initialize()
+	{
+		UTwoAxisPropertyEditInputBehavior::Initialize(this);
+	}
+
+	FPropertyInterface HorizontalProperty;
+	FPropertyInterface VerticalProperty;
+	TUniqueFunction<FKey()> GetCaptureKeyFunc = [](){ return EKeys::B; };
+	FSimpleMulticastDelegate OnDragUpdated;
+
+private:
+
+	// ITwoAxisPropertyEditBehaviorTarget implementation
+
+	virtual IPropertyInterface* GetHorizontalProperty() override { return &HorizontalProperty; }
+	virtual IPropertyInterface* GetVerticalProperty() override { return &VerticalProperty; }
+	virtual FKey GetCaptureKey() override { return GetCaptureKeyFunc(); }
+	virtual void PostDragUpdated() override { OnDragUpdated.Broadcast(); }
+};
\ No newline at end of file
diff --git a/Engine/Source/Runtime/InteractiveToolsFramework/Public/BaseTools/BaseBrushTool.h b/Engine/Source/Runtime/InteractiveToolsFramework/Public/BaseTools/BaseBrushTool.h
--- a/Engine/Source/Runtime/InteractiveToolsFramework/Public/BaseTools/BaseBrushTool.h
+++ b/Engine/Source/Runtime/InteractiveToolsFramework/Public/BaseTools/BaseBrushTool.h
@@ -9,6 +9,8 @@
 #include "BaseBrushTool.generated.h"
 
 class UBrushStampIndicator;
+class ULocalTwoAxisPropertyEditInputBehavior;
+
 /**
  * Standard properties for a Brush-type Tool
  */
@@ -82,7 +84,7 @@ struct FBrushStampData
  * vertical and horizontal directions respectively.
  */
 UCLASS()
-class UBrushAdjusterInputBehavior : public UAnyButtonInputBehavior
+class UE_DEPRECATED(5.5, "UBrushAdjusterInputBehavior is deprecated, use UTwoAxisPropertyEditInputBehavior instead.") UBrushAdjusterInputBehavior : public UAnyButtonInputBehavior
 {
 	GENERATED_BODY()
 
@@ -242,7 +244,12 @@ protected:
 	INTERACTIVETOOLSFRAMEWORK_API virtual void ShutdownBrushStampIndicator();
 
 	// adjusts size and strength properties while holding hotkey during click/drag
+	PRAGMA_DISABLE_DEPRECATION_WARNINGS
+	UE_DEPRECATED(5.5, "BrushAdjusterBehavior is deprecated, use BrushEditBehavior instead.")
 	TWeakObjectPtr<UBrushAdjusterInputBehavior> BrushAdjusterBehavior;
+	PRAGMA_ENABLE_DEPRECATION_WARNINGS
+	
+	TWeakObjectPtr<ULocalTwoAxisPropertyEditInputBehavior> BrushEditBehavior;
 
 private:
 	
