New TRS Gizmo: fix hit testing for cylinder and cones in orthographic views

#rnx

diff --git a/Engine/Source/Runtime/InteractiveToolsFramework/Private/BaseGizmos/GizmoElementCone.cpp b/Engine/Source/Runtime/InteractiveToolsFramework/Private/BaseGizmos/GizmoElementCone.cpp
--- a/Engine/Source/Runtime/InteractiveToolsFramework/Private/BaseGizmos/GizmoElementCone.cpp
+++ b/Engine/Source/Runtime/InteractiveToolsFramework/Private/BaseGizmos/GizmoElementCone.cpp
@@ -47,12 +47,20 @@ FInputRayHit UGizmoElementCone::LineTrace(const UGizmoViewContext* ViewContext,
 		const FVector WorldDirection = CurrentLineTraceState.LocalToWorldTransform.TransformVectorNoScale(Direction);
 		const FVector WorldOrigin = CurrentLineTraceState.LocalToWorldTransform.TransformPosition(FVector::ZeroVector) - WorldDirection * PixelHitThresholdAdjust;
 
+		// due to numerical imprecision, the ray origin needs to be clamped in ortho views
+		// (cf. UEditorInteractiveToolsContext::GetRayFromMousePos)
+		FVector ClampedRayOrigin(RayOrigin);
+		if (!ViewContext->IsPerspectiveProjection())
+		{
+			ClampedRayOrigin.Z = FMath::Clamp(RayOrigin.Z, -UE_OLD_HALF_WORLD_MAX, UE_OLD_HALF_WORLD_MAX);
+		}
+		
 		GizmoMath::RayConeIntersection(
 			WorldOrigin,
 			WorldDirection,
 			CosAngle,
 			WorldHeight,
-			RayOrigin, RayDirection,
+			ClampedRayOrigin, RayDirection,
 			bIntersects, RayParam);
 
 		if (bIntersects)
diff --git a/Engine/Source/Runtime/InteractiveToolsFramework/Private/BaseGizmos/GizmoElementCylinder.cpp b/Engine/Source/Runtime/InteractiveToolsFramework/Private/BaseGizmos/GizmoElementCylinder.cpp
--- a/Engine/Source/Runtime/InteractiveToolsFramework/Private/BaseGizmos/GizmoElementCylinder.cpp
+++ b/Engine/Source/Runtime/InteractiveToolsFramework/Private/BaseGizmos/GizmoElementCylinder.cpp
@@ -47,12 +47,20 @@ FInputRayHit UGizmoElementCylinder::LineTrace(const UGizmoViewContext* ViewConte
 		const FVector LocalCenter = Direction * Height * 0.5;
 		const FVector WorldCenter = CurrentLineTraceState.LocalToWorldTransform.TransformPosition(LocalCenter);
 
+		// due to numerical imprecision, the ray origin needs to be clamped in ortho views
+		// (cf. UEditorInteractiveToolsContext::GetRayFromMousePos)
+		FVector ClampedRayOrigin(RayOrigin);
+		if (!ViewContext->IsPerspectiveProjection())
+		{
+			ClampedRayOrigin.Z = FMath::Clamp(RayOrigin.Z, -UE_OLD_HALF_WORLD_MAX, UE_OLD_HALF_WORLD_MAX);
+		}
+		
 		GizmoMath::RayCylinderIntersection(
 			WorldCenter,
 			WorldDirection,
 			WorldRadius,
 			WorldHeight,
-			RayOrigin, RayDirection,
+			ClampedRayOrigin, RayDirection,
 			bIntersects, RayParam);
 
 		if (bIntersects)
diff --git a/Engine/Source/Runtime/InteractiveToolsFramework/Private/BaseGizmos/GizmoElementTorus.cpp b/Engine/Source/Runtime/InteractiveToolsFramework/Private/BaseGizmos/GizmoElementTorus.cpp
--- a/Engine/Source/Runtime/InteractiveToolsFramework/Private/BaseGizmos/GizmoElementTorus.cpp
+++ b/Engine/Source/Runtime/InteractiveToolsFramework/Private/BaseGizmos/GizmoElementTorus.cpp
@@ -106,6 +106,14 @@ FInputRayHit UGizmoElementTorus::LineTrace(const UGizmoViewContext* ViewContext,
 				return FInputRayHit();
 			}
 
+			// due to numerical imprecision, the ray origin needs to be clamped in ortho views
+			// (cf. UEditorInteractiveToolsContext::GetRayFromMousePos)
+			FVector ClampedRayOrigin(RayOrigin);
+			if (!ViewContext->IsPerspectiveProjection())
+			{
+				ClampedRayOrigin.Z = FMath::Clamp(RayOrigin.Z, -UE_OLD_HALF_WORLD_MAX, UE_OLD_HALF_WORLD_MAX);
+			}
+
 			// Line trace against a set of cylinders approximating the shape of the torus
 			for (int32 i = 0; i < NumCylinders; i++)
 			{
@@ -130,7 +138,7 @@ FInputRayHit UGizmoElementTorus::LineTrace(const UGizmoViewContext* ViewContext,
 					CylinderDirection,
 					CylinderRadius,
 					CylinderHeight,
-					RayOrigin, RayDirection,
+					ClampedRayOrigin, RayDirection,
 					bIntersects, RayParam);
 
 				// Update with closest hit depth
