New TRS Gizmo: add support for masking in arrow elements to adapt the threshold based on another element
- this is usefull for translate & scale axis to reduce the threshold when the screen space / uniform gizmo element is hit

#rnx

diff --git a/Engine/Source/Editor/Experimental/EditorInteractiveToolsFramework/Private/EditorGizmos/TransformGizmo.cpp b/Engine/Source/Editor/Experimental/EditorInteractiveToolsFramework/Private/EditorGizmos/TransformGizmo.cpp
--- a/Engine/Source/Editor/Experimental/EditorInteractiveToolsFramework/Private/EditorGizmos/TransformGizmo.cpp
+++ b/Engine/Source/Editor/Experimental/EditorInteractiveToolsFramework/Private/EditorGizmos/TransformGizmo.cpp
@@ -611,19 +611,19 @@ void UTransformGizmo::EnableTranslate(EAxisList::Type InAxisListToDraw)
 
 	if (bEnableX && TranslateXAxisElement == nullptr)
 	{
-		TranslateXAxisElement = MakeTranslateAxis(ETransformGizmoPartIdentifier::TranslateXAxis, FVector(1.0f, 0.0f, 0.0f), FVector(0.0f, 1.0f, 0.0f), AxisMaterialX);
+		TranslateXAxisElement = MakeTranslateAxis(ETransformGizmoPartIdentifier::TranslateXAxis, FVector::XAxisVector, FVector::YAxisVector, AxisMaterialX);
 		GizmoElementRoot->Add(TranslateXAxisElement);
 	}
 
 	if (bEnableY && TranslateYAxisElement == nullptr)
 	{
-		TranslateYAxisElement = MakeTranslateAxis(ETransformGizmoPartIdentifier::TranslateYAxis, FVector(0.0f, 1.0f, 0.0f), FVector(0.0f, 0.0f, 1.0f), AxisMaterialY);
+		TranslateYAxisElement = MakeTranslateAxis(ETransformGizmoPartIdentifier::TranslateYAxis, FVector::YAxisVector, FVector::ZAxisVector, AxisMaterialY);
 		GizmoElementRoot->Add(TranslateYAxisElement);
 	}
 
 	if (bEnableZ && TranslateZAxisElement == nullptr)
 	{
-		TranslateZAxisElement = MakeTranslateAxis(ETransformGizmoPartIdentifier::TranslateZAxis, FVector(0.0f, 0.0f, 1.0f), FVector(1.0f, 0.0f, 0.0f), AxisMaterialZ);
+		TranslateZAxisElement = MakeTranslateAxis(ETransformGizmoPartIdentifier::TranslateZAxis, FVector::ZAxisVector, FVector::XAxisVector, AxisMaterialZ);
 		GizmoElementRoot->Add(TranslateZAxisElement);
 	}
 
@@ -636,16 +636,19 @@ void UTransformGizmo::EnableTranslate(EAxisList::Type InAxisListToDraw)
 	if (TranslateXAxisElement)
 	{
 		TranslateXAxisElement->SetEnabled(bEnableX);
+		TranslateXAxisElement->SetHitMask(TranslateScreenSpaceElement);
 	}
 
 	if (TranslateYAxisElement)
 	{
 		TranslateYAxisElement->SetEnabled(bEnableY);
+		TranslateYAxisElement->SetHitMask(TranslateScreenSpaceElement);
 	}
 
 	if (TranslateZAxisElement)
 	{
 		TranslateZAxisElement->SetEnabled(bEnableZ);
+		TranslateZAxisElement->SetHitMask(TranslateScreenSpaceElement);
 	}
 
 	if (TranslateScreenSpaceElement)
diff --git a/Engine/Source/Runtime/InteractiveToolsFramework/Private/BaseGizmos/GizmoElementArrow.cpp b/Engine/Source/Runtime/InteractiveToolsFramework/Private/BaseGizmos/GizmoElementArrow.cpp
--- a/Engine/Source/Runtime/InteractiveToolsFramework/Private/BaseGizmos/GizmoElementArrow.cpp
+++ b/Engine/Source/Runtime/InteractiveToolsFramework/Private/BaseGizmos/GizmoElementArrow.cpp
@@ -4,13 +4,27 @@
 #include "BaseGizmos/GizmoElementBox.h"
 #include "BaseGizmos/GizmoElementCone.h"
 #include "BaseGizmos/GizmoElementCylinder.h"
-#include "BaseGizmos/GizmoMath.h"
 #include "BaseGizmos/GizmoRenderingUtil.h"
 #include "InputState.h"
 #include "Materials/MaterialInterface.h"
 
 #include UE_INLINE_GENERATED_CPP_BY_NAME(GizmoElementArrow)
 
+namespace UE::Private
+{
+
+void UpdatePixelThreshold(UGizmoElementBase* InGizmoElement, const float InPixelHitThreshold)
+{
+	if (InGizmoElement)
+	{
+		InGizmoElement->SetPixelHitDistanceThreshold(InPixelHitThreshold);
+	}		
+}
+
+using ThresholdGuard = TGuardValue_Bitfield_Cleanup<TFunction<void()>>;
+
+}
+
 UGizmoElementArrow::UGizmoElementArrow()
 {
 	HeadType = EGizmoElementArrowHeadType::Cone;
@@ -56,6 +70,8 @@ void UGizmoElementArrow::Render(IToolsContextRenderAPI* RenderAPI, const FRender
 
 FInputRayHit UGizmoElementArrow::LineTrace(const UGizmoViewContext* ViewContext, const FLineTraceTraversalState& LineTraceState, const FVector& RayOrigin, const FVector& RayDirection)
 {
+	using namespace UE::Private;
+	
 	FLineTraceTraversalState CurrentLineTraceState(LineTraceState);
 	bool bHittableViewDependent = UpdateLineTraceState(ViewContext, Base, CurrentLineTraceState);
 
@@ -64,6 +80,26 @@ FInputRayHit UGizmoElementArrow::LineTrace(const UGizmoViewContext* ViewContext,
 		return FInputRayHit();
 	}
 
+	ThresholdGuard Guard([this]()
+	{
+		UpdatePixelThreshold(CylinderElement, PixelHitDistanceThreshold);
+		UpdatePixelThreshold(ConeElement, PixelHitDistanceThreshold);
+		UpdatePixelThreshold(BoxElement, PixelHitDistanceThreshold);
+	});
+
+	// update threshold if hitting the mask
+	if (UGizmoElementBase* HitMaskGizmo = HitMask.Get())
+	{
+		const FInputRayHit MaskHit = HitMaskGizmo->LineTrace(ViewContext, LineTraceState, RayOrigin, RayDirection);
+		if (MaskHit.bHit)
+		{
+			static constexpr float NoPixelHitThreshold = 0.f;
+			UpdatePixelThreshold(CylinderElement, NoPixelHitThreshold);
+			UpdatePixelThreshold(ConeElement, NoPixelHitThreshold);
+			UpdatePixelThreshold(BoxElement, NoPixelHitThreshold);		
+		}
+	}
+	
 	check(CylinderElement);
 	FInputRayHit Hit = CylinderElement->LineTrace(ViewContext, CurrentLineTraceState, RayOrigin, RayDirection);
 
@@ -213,6 +249,11 @@ void UGizmoElementArrow::SetPixelHitDistanceThreshold(float InPixelHitDistanceTh
 	}
 }
 
+void UGizmoElementArrow::SetHitMask(const TWeakObjectPtr<UGizmoElementBase>& InHitMask)
+{
+	HitMask = InHitMask;
+}
+
 void UGizmoElementArrow::SetHeadType(EGizmoElementArrowHeadType InHeadType)
 {
 	if (InHeadType != HeadType)
diff --git a/Engine/Source/Runtime/InteractiveToolsFramework/Public/BaseGizmos/GizmoElementArrow.h b/Engine/Source/Runtime/InteractiveToolsFramework/Public/BaseGizmos/GizmoElementArrow.h
--- a/Engine/Source/Runtime/InteractiveToolsFramework/Public/BaseGizmos/GizmoElementArrow.h
+++ b/Engine/Source/Runtime/InteractiveToolsFramework/Public/BaseGizmos/GizmoElementArrow.h
@@ -77,6 +77,9 @@ public:
 	// Pixel hit distance threshold, element will be scaled enough to add this threshold when line-tracing.
 	INTERACTIVETOOLSFRAMEWORK_API virtual void SetPixelHitDistanceThreshold(float InPixelHitDistanceThreshold) override;
 
+	// Hit Mask Gizmo element used for adaptive pixel hit threshold
+	INTERACTIVETOOLSFRAMEWORK_API void SetHitMask(const TWeakObjectPtr<UGizmoElementBase>& InHitMask);
+	
 protected:
 
 	// Update arrow cylinder based on parameters
@@ -138,4 +141,8 @@ protected:
 	// Head type
 	UPROPERTY()
 	EGizmoElementArrowHeadType HeadType = EGizmoElementArrowHeadType::Cone;
+
+	// Hit mask
+	UPROPERTY()
+	TWeakObjectPtr<UGizmoElementBase> HitMask = nullptr;
 };
