fix issue where free rotation could (rarely) get 'stuck' (due to the internal rotation representation being filled w/ NaN values)

diff --git a/Engine/Source/Runtime/InteractiveToolsFramework/Private/BaseGizmos/FreeRotationSubGizmo.cpp b/Engine/Source/Runtime/InteractiveToolsFramework/Private/BaseGizmos/FreeRotationSubGizmo.cpp
--- a/Engine/Source/Runtime/InteractiveToolsFramework/Private/BaseGizmos/FreeRotationSubGizmo.cpp
+++ b/Engine/Source/Runtime/InteractiveToolsFramework/Private/BaseGizmos/FreeRotationSubGizmo.cpp
@@ -184,7 +184,7 @@ bool UFreeRotationSubGizmo::ClickDrag_IncrementalSphereBound(const FInputDeviceR
 	{
 		return false;
 	}
-	double RotationAngle = acos(ToLastSpherePoint.Dot(ToCurrentSpherePoint));
+	double RotationAngle = UE::Geometry::AngleR(ToLastSpherePoint, ToCurrentSpherePoint);
 	
 	// Apply the new transform to the transform source
 	FTransform CurrentTransform = TransformSource->GetTransform();
