New TRS Gizmo: fix hit testing for cylinder and cones in orthographic views
- due to numerical imprecision, the ray origin sometimes needs to be clamped in ortho views (cf. UEditorInteractiveToolsContext::GetRayFromMousePos)

#rnx

diff --git a/Engine/Source/Runtime/InteractiveToolsFramework/Private/BaseGizmos/GizmoElementCone.cpp b/Engine/Source/Runtime/InteractiveToolsFramework/Private/BaseGizmos/GizmoElementCone.cpp
--- a/Engine/Source/Runtime/InteractiveToolsFramework/Private/BaseGizmos/GizmoElementCone.cpp
+++ b/Engine/Source/Runtime/InteractiveToolsFramework/Private/BaseGizmos/GizmoElementCone.cpp
@@ -1,6 +1,8 @@
 // Copyright Epic Games, Inc. All Rights Reserved.
  
 #include "BaseGizmos/GizmoElementCone.h"
+
+#include "GizmoPrivateUtil.h"
 #include "BaseGizmos/GizmoRenderingUtil.h"
 #include "BaseGizmos/GizmoMath.h"
 #include "InputState.h"
@@ -50,10 +52,7 @@ FInputRayHit UGizmoElementCone::LineTrace(const UGizmoViewContext* ViewContext,
 		// due to numerical imprecision, the ray origin needs to be clamped in ortho views
 		// (cf. UEditorInteractiveToolsContext::GetRayFromMousePos)
 		FVector ClampedRayOrigin(RayOrigin);
-		if (!ViewContext->IsPerspectiveProjection())
-		{
-			ClampedRayOrigin.Z = FMath::Clamp(RayOrigin.Z, -UE_OLD_HALF_WORLD_MAX, UE_OLD_HALF_WORLD_MAX);
-		}
+		const double DepthBias = UE::GizmoUtil::ClampRayOrigin(ViewContext, ClampedRayOrigin, RayDirection);
 		
 		GizmoMath::RayConeIntersection(
 			WorldOrigin,
@@ -65,6 +64,8 @@ FInputRayHit UGizmoElementCone::LineTrace(const UGizmoViewContext* ViewContext,
 
 		if (bIntersects)
 		{
+			// add the depth bias if any
+			RayParam += DepthBias;
 			FInputRayHit RayHit(static_cast<float>(RayParam));
 			RayHit.SetHitObject(this);
 			RayHit.HitIdentifier = PartIdentifier;
diff --git a/Engine/Source/Runtime/InteractiveToolsFramework/Private/BaseGizmos/GizmoElementCylinder.cpp b/Engine/Source/Runtime/InteractiveToolsFramework/Private/BaseGizmos/GizmoElementCylinder.cpp
--- a/Engine/Source/Runtime/InteractiveToolsFramework/Private/BaseGizmos/GizmoElementCylinder.cpp
+++ b/Engine/Source/Runtime/InteractiveToolsFramework/Private/BaseGizmos/GizmoElementCylinder.cpp
@@ -1,6 +1,8 @@
 // Copyright Epic Games, Inc. All Rights Reserved.
  
 #include "BaseGizmos/GizmoElementCylinder.h"
+
+#include "GizmoPrivateUtil.h"
 #include "BaseGizmos/GizmoRenderingUtil.h"
 #include "BaseGizmos/GizmoMath.h"
 #include "InputState.h"
@@ -50,10 +52,7 @@ FInputRayHit UGizmoElementCylinder::LineTrace(const UGizmoViewContext* ViewConte
 		// due to numerical imprecision, the ray origin needs to be clamped in ortho views
 		// (cf. UEditorInteractiveToolsContext::GetRayFromMousePos)
 		FVector ClampedRayOrigin(RayOrigin);
-		if (!ViewContext->IsPerspectiveProjection())
-		{
-			ClampedRayOrigin.Z = FMath::Clamp(RayOrigin.Z, -UE_OLD_HALF_WORLD_MAX, UE_OLD_HALF_WORLD_MAX);
-		}
+		const double DepthBias = UE::GizmoUtil::ClampRayOrigin(ViewContext, ClampedRayOrigin, RayDirection);
 		
 		GizmoMath::RayCylinderIntersection(
 			WorldCenter,
@@ -65,6 +64,8 @@ FInputRayHit UGizmoElementCylinder::LineTrace(const UGizmoViewContext* ViewConte
 
 		if (bIntersects)
 		{
+			// add the depth bias if any
+			RayParam += DepthBias;
 			FInputRayHit RayHit(static_cast<float>(RayParam));
 			RayHit.SetHitObject(this);
 			RayHit.HitIdentifier = PartIdentifier;
diff --git a/Engine/Source/Runtime/InteractiveToolsFramework/Private/BaseGizmos/GizmoElementTorus.cpp b/Engine/Source/Runtime/InteractiveToolsFramework/Private/BaseGizmos/GizmoElementTorus.cpp
--- a/Engine/Source/Runtime/InteractiveToolsFramework/Private/BaseGizmos/GizmoElementTorus.cpp
+++ b/Engine/Source/Runtime/InteractiveToolsFramework/Private/BaseGizmos/GizmoElementTorus.cpp
@@ -1,6 +1,8 @@
 // Copyright Epic Games, Inc. All Rights Reserved.
  
 #include "BaseGizmos/GizmoElementTorus.h"
+
+#include "GizmoPrivateUtil.h"
 #include "BaseGizmos/GizmoMath.h"
 #include "InputState.h"
 #include "Materials/MaterialInterface.h"
@@ -109,11 +111,10 @@ FInputRayHit UGizmoElementTorus::LineTrace(const UGizmoViewContext* ViewContext,
 			// due to numerical imprecision, the ray origin needs to be clamped in ortho views
 			// (cf. UEditorInteractiveToolsContext::GetRayFromMousePos)
 			FVector ClampedRayOrigin(RayOrigin);
-			if (!ViewContext->IsPerspectiveProjection())
-			{
-				ClampedRayOrigin.Z = FMath::Clamp(RayOrigin.Z, -UE_OLD_HALF_WORLD_MAX, UE_OLD_HALF_WORLD_MAX);
-			}
+			const double DepthBias = UE::GizmoUtil::ClampRayOrigin(ViewContext, ClampedRayOrigin, RayDirection);
 
+			bool bDidIntersect = false;
+			
 			// Line trace against a set of cylinders approximating the shape of the torus
 			for (int32 i = 0; i < NumCylinders; i++)
 			{
@@ -141,12 +142,20 @@ FInputRayHit UGizmoElementTorus::LineTrace(const UGizmoViewContext* ViewContext,
 					ClampedRayOrigin, RayDirection,
 					bIntersects, RayParam);
 
+				bDidIntersect |= bIntersects; 
+				
 				// Update with closest hit depth
 				if (bIntersects && (HitDepth < 0 || HitDepth > RayParam))
 				{
 					HitDepth = RayParam;
 				}
 			}
+
+			if (bDidIntersect)
+			{
+				// add the depth bias if any
+				HitDepth += DepthBias;
+			}
 		}
 		else
 		{
diff --git a/Engine/Source/Runtime/InteractiveToolsFramework/Private/BaseGizmos/GizmoPrivateUtil.cpp b/Engine/Source/Runtime/InteractiveToolsFramework/Private/BaseGizmos/GizmoPrivateUtil.cpp
--- a/Engine/Source/Runtime/InteractiveToolsFramework/Private/BaseGizmos/GizmoPrivateUtil.cpp
+++ b/Engine/Source/Runtime/InteractiveToolsFramework/Private/BaseGizmos/GizmoPrivateUtil.cpp
@@ -219,6 +219,28 @@ bool UE::GizmoUtil::SetCommonSubGizmoProperties(
 	return true;
 }
 
+double UE::GizmoUtil::ClampRayOrigin(const UGizmoViewContext* InViewContext, FVector& InOutRayOrigin, const FVector& InRayDirection, const double InMaxDepth)
+{
+	// due to numerical imprecision, the ray origin needs to be clamped in ortho views
+	// (cf. UEditorInteractiveToolsContext::GetRayFromMousePos)
+	
+	if (!InViewContext->IsPerspectiveProjection())
+	{
+		const double dot = FVector::DotProduct(InOutRayOrigin, InRayDirection);
+		if (FMath::Abs(dot) > InMaxDepth)
+		{
+			InOutRayOrigin += -dot * InRayDirection;
+			InOutRayOrigin -= InMaxDepth * InRayDirection;
+			
+			const double DepthBias = (-dot) - InMaxDepth;
+			return DepthBias;
+		}
+	}
+	
+	return 0.0;
+}
+
+
 // Explicit instantiations to avoid having the definition in the header
 template bool UE::GizmoUtil::SetCommonSubGizmoProperties<UAxisPositionGizmo>(
 	UAxisPositionGizmo* Gizmo,
diff --git a/Engine/Source/Runtime/InteractiveToolsFramework/Private/BaseGizmos/GizmoPrivateUtil.h b/Engine/Source/Runtime/InteractiveToolsFramework/Private/BaseGizmos/GizmoPrivateUtil.h
--- a/Engine/Source/Runtime/InteractiveToolsFramework/Private/BaseGizmos/GizmoPrivateUtil.h
+++ b/Engine/Source/Runtime/InteractiveToolsFramework/Private/BaseGizmos/GizmoPrivateUtil.h
@@ -55,4 +55,11 @@ namespace UE::GizmoUtil
 		const UE::GizmoUtil::FTransformSubGizmoCommonParams& Params,
 		UE::GizmoUtil::FTransformSubGizmoSharedState* SharedState,
 		UGizmoScaledAndUnscaledTransformSources*& TransformSourceOut);
+
+	/**
+	 * Helper to ensure that the ray origin is not too far in orthographic views to avoid numerical imprecision later
+	 * @return the depth bias between the current ray origin and the updated one if the ray origin is greater than InMaxDepth
+	 */
+	double ClampRayOrigin(const UGizmoViewContext* InViewContext, FVector& InOutRayOrigin, const FVector& InRayDirection, const double InMaxDepth = UE_OLD_HALF_WORLD_MAX);
+	
 }// end UE::GizmoUtil
\ No newline at end of file
