Fix single- and multi- selection tool focus logic for targets that implement the scene/primitive component backed interface without using the uprimitivecomponenttooltarget base class

diff --git a/Engine/Source/Runtime/InteractiveToolsFramework/Private/BaseTools/SingleSelectionTool.cpp b/Engine/Source/Runtime/InteractiveToolsFramework/Private/BaseTools/SingleSelectionTool.cpp
--- a/Engine/Source/Runtime/InteractiveToolsFramework/Private/BaseTools/SingleSelectionTool.cpp
+++ b/Engine/Source/Runtime/InteractiveToolsFramework/Private/BaseTools/SingleSelectionTool.cpp
@@ -11,7 +11,7 @@
 
 bool USingleSelectionTool::SupportsWorldSpaceFocusBox()
 {
-	return Target && Cast<UPrimitiveComponentToolTarget>(Target) != nullptr;
+	return Target && Cast<ISceneComponentBackedTarget>(Target) != nullptr;
 }
 
 
@@ -19,10 +19,10 @@ FBox USingleSelectionTool::GetWorldSpaceFocusBox()
 {
 	if (Target)
 	{
-		UPrimitiveComponentToolTarget* PrimTarget = Cast<UPrimitiveComponentToolTarget>(Target);
+		ISceneComponentBackedTarget* PrimTarget = Cast<ISceneComponentBackedTarget>(Target);
 		if (PrimTarget)
 		{
-			UPrimitiveComponent* Component = PrimTarget->GetOwnerComponent();
+			USceneComponent* Component = PrimTarget->GetOwnerSceneComponent();
 			if (Component)
 			{
 				return Component->Bounds.GetBox();
@@ -35,14 +35,14 @@ FBox USingleSelectionTool::GetWorldSpaceFocusBox()
 
 bool USingleSelectionTool::SupportsWorldSpaceFocusPoint()
 {
-	return Target && Cast<UPrimitiveComponentToolTarget>(Target) != nullptr;
+	return Target && Cast<IPrimitiveComponentBackedTarget>(Target) != nullptr;
 
 }
 bool USingleSelectionTool::GetWorldSpaceFocusPoint(const FRay& WorldRay, FVector& PointOut)
 {
 	if (Target)
 	{
-		UPrimitiveComponentToolTarget* PrimTarget = Cast<UPrimitiveComponentToolTarget>(Target);
+		IPrimitiveComponentBackedTarget* PrimTarget = Cast<IPrimitiveComponentBackedTarget>(Target);
 		if (PrimTarget)
 		{
 			FHitResult HitResult;
diff --git a/Engine/Source/Runtime/InteractiveToolsFramework/Private/MultiSelectionTool.cpp b/Engine/Source/Runtime/InteractiveToolsFramework/Private/MultiSelectionTool.cpp
--- a/Engine/Source/Runtime/InteractiveToolsFramework/Private/MultiSelectionTool.cpp
+++ b/Engine/Source/Runtime/InteractiveToolsFramework/Private/MultiSelectionTool.cpp
@@ -11,7 +11,7 @@ bool UMultiSelectionTool::SupportsWorldSpaceFocusBox()
 	int32 PrimitiveCount = 0;
 	for (const TObjectPtr<UToolTarget>& Target : Targets)
 	{
-		if (Cast<UPrimitiveComponentToolTarget>(Target) != nullptr)
+		if (Cast<ISceneComponentBackedTarget>(Target) != nullptr)
 		{
 			PrimitiveCount++;
 		}
@@ -25,10 +25,10 @@ FBox UMultiSelectionTool::GetWorldSpaceFocusBox()
 	FBox AccumBox(EForceInit::ForceInit);
 	for (const TObjectPtr<UToolTarget>& Target : Targets)
 	{
-		UPrimitiveComponentToolTarget* PrimTarget = Cast<UPrimitiveComponentToolTarget>(Target);
+		ISceneComponentBackedTarget* PrimTarget = Cast<ISceneComponentBackedTarget>(Target);
 		if (PrimTarget)
 		{
-			UPrimitiveComponent* Component = PrimTarget->GetOwnerComponent();
+			USceneComponent* Component = PrimTarget->GetOwnerSceneComponent();
 			if (Component)
 			{
 				FBox ComponentBounds = Component->Bounds.GetBox();
@@ -45,7 +45,7 @@ bool UMultiSelectionTool::SupportsWorldSpaceFocusPoint()
 	int32 PrimitiveCount = 0;
 	for (const TObjectPtr<UToolTarget>& Target : Targets)
 	{
-		if (Cast<UPrimitiveComponentToolTarget>(Target) != nullptr)
+		if (Cast<IPrimitiveComponentBackedTarget>(Target) != nullptr)
 		{
 			PrimitiveCount++;
 		}
@@ -60,7 +60,7 @@ bool UMultiSelectionTool::GetWorldSpaceFocusPoint(const FRay& WorldRay, FVector&
 
 	for (const TObjectPtr<UToolTarget>& Target : Targets)
 	{
-		UPrimitiveComponentToolTarget* PrimTarget = Cast<UPrimitiveComponentToolTarget>(Target);
+		IPrimitiveComponentBackedTarget* PrimTarget = Cast<IPrimitiveComponentBackedTarget>(Target);
 		if (PrimTarget)
 		{
 			FHitResult HitResult;
