InteractiveToolsFramework: Added function to get the ActiveToolBuilderName from the ToolManager.

diff --git a/Engine/Source/Runtime/InteractiveToolsFramework/Private/InteractiveToolManager.cpp b/Engine/Source/Runtime/InteractiveToolsFramework/Private/InteractiveToolManager.cpp
--- a/Engine/Source/Runtime/InteractiveToolsFramework/Private/InteractiveToolManager.cpp
+++ b/Engine/Source/Runtime/InteractiveToolsFramework/Private/InteractiveToolManager.cpp
@@ -112,6 +112,11 @@ bool UInteractiveToolManager::SelectActiveToolType(EToolSide Side, const FString
 }
 
 
+const FString& UInteractiveToolManager::GetActiveToolType(EToolSide Side) const
+{
+	return Side == EToolSide::Right ? ActiveRightBuilderName : ActiveLeftBuilderName;
+}
+
 
 bool UInteractiveToolManager::CanActivateTool(EToolSide Side, FString Identifier)
 {
diff --git a/Engine/Source/Runtime/InteractiveToolsFramework/Public/InteractiveToolManager.h b/Engine/Source/Runtime/InteractiveToolsFramework/Public/InteractiveToolManager.h
--- a/Engine/Source/Runtime/InteractiveToolsFramework/Public/InteractiveToolManager.h
+++ b/Engine/Source/Runtime/InteractiveToolsFramework/Public/InteractiveToolManager.h
@@ -130,6 +130,12 @@ public:
 	 */
 	INTERACTIVETOOLSFRAMEWORK_API virtual bool SelectActiveToolType(EToolSide Side, const FString& Identifier);
 
+	/**
+	 * Get active ToolBuilder string identifier for a ToolSide
+	 * @param Side which "side" should we get this Builder from 
+	 */
+	INTERACTIVETOOLSFRAMEWORK_API virtual const FString& GetActiveToolType(EToolSide Side) const;
+
 	/**
 	 * Check if a named Tool type can currently be activated on the given ToolSide
 	 * @param Side which "side" you would like to active the tool on
