[ITF] Added DrawArc functions to ToolDataVisualizer

diff --git a/Engine/Source/Runtime/InteractiveToolsFramework/Private/ToolDataVisualizer.cpp b/Engine/Source/Runtime/InteractiveToolsFramework/Private/ToolDataVisualizer.cpp
--- a/Engine/Source/Runtime/InteractiveToolsFramework/Private/ToolDataVisualizer.cpp
+++ b/Engine/Source/Runtime/InteractiveToolsFramework/Private/ToolDataVisualizer.cpp
@@ -113,6 +113,19 @@ void FToolDataVisualizer::InternalDrawCircle(const FVector& Position, const FVec
 		LineThicknessIn * PDISizeScale, DepthBias, true);
 }
 
+void FToolDataVisualizer::InternalDrawArc(const FVector& InPosition, const FVector& InNormal, const float InRadius, const int InSteps, const float InMinAngle, const float InMaxAngle, const FLinearColor& Color, float LineThicknessIn, const bool bDepthTestedIn)
+{
+	FVector Tan1, Tan2;
+	GizmoMath::MakeNormalPlaneBasis(TransformN(InNormal), Tan1, Tan2);
+	Tan1.Normalize(); Tan2.Normalize();
+
+	// this function is from SceneManagement.h
+	::DrawArc(CurrentPDI, TransformP(InPosition), Tan1, Tan2,
+		InMinAngle, InMaxAngle,
+		InRadius, InSteps, Color,
+		static_cast<uint8>(bDepthTestedIn ? SDPG_World : SDPG_Foreground));
+}
+
 void FToolDataVisualizer::InternalDrawWireBox(const FBox& Box, const FLinearColor& ColorIn, float LineThicknessIn, bool bDepthTestedIn)
 {
 	// corners [ (-x,-y), (x,-y), (x,y), (-x,y) ], -z, then +z
@@ -197,6 +210,24 @@ void FToolDataVisualizer::InternalDrawViewFacingCircle(const FVector& Position,
 		LineThicknessIn * PDISizeScale, DepthBias, true);
 }
 
+
+void FToolDataVisualizer::InternalDrawViewFacingArc(const FVector& InPosition, const float InRadius, const int InSteps, const float InMinAngle, const float InMaxAngle, const FLinearColor& Color, float LineThicknessIn, const bool bDepthTestedIn)
+{
+	checkf(bHaveCameraState, TEXT("To call this function, you must first call the version of BeginFrame that takes the CameraState"));
+
+	const FVector WorldPosition = TransformP(InPosition);
+	FVector WorldNormal = (CameraState.Position - WorldPosition);
+	WorldNormal.Normalize();
+	FVector Tan1, Tan2;
+	GizmoMath::MakeNormalPlaneBasis(WorldNormal, Tan1, Tan2);
+
+	// this function is from SceneManagement.h
+	::DrawArc(CurrentPDI, WorldPosition, Tan1, Tan2,
+		InMinAngle, InMaxAngle,
+		InRadius, InSteps, Color,
+		static_cast<uint8>(bDepthTestedIn ? SDPG_World : SDPG_Foreground));
+}
+ 
 void FToolDataVisualizer::InternalDrawViewFacingX(const FVector& Position, float Width, const FLinearColor& Color, float LineThicknessIn, bool bDepthTestedIn)
 {
 	checkf(bHaveCameraState, TEXT("To call this function, you must first call the version of BeginFrame that takes the CameraState"));
diff --git a/Engine/Source/Runtime/InteractiveToolsFramework/Public/ToolDataVisualizer.h b/Engine/Source/Runtime/InteractiveToolsFramework/Public/ToolDataVisualizer.h
--- a/Engine/Source/Runtime/InteractiveToolsFramework/Public/ToolDataVisualizer.h
+++ b/Engine/Source/Runtime/InteractiveToolsFramework/Public/ToolDataVisualizer.h
@@ -180,6 +180,20 @@ public:
 		InternalDrawViewFacingCircle((FVector)Position, Radius, Steps, Color, LineThicknessIn, bDepthTestedIn);
 	}
 
+	/** Draw a 3D arc at given position/normal with the given parameters. Tangent axes are defined internally */
+	template<typename PointType>
+	void DrawArc(const PointType& Position, const PointType& Normal, float Radius, int Steps, float MinAngle, float MaxAngle, const FLinearColor& Color, float LineThicknessIn, bool bDepthTestedIn)
+	{
+		InternalDrawArc((FVector)Position, (FVector)Normal, Radius, Steps, MinAngle, MaxAngle, Color, LineThicknessIn, bDepthTestedIn);
+	}
+
+	/** Draw a 3D arc at given position/normal with the given parameters. Tangent axes are defined internally */
+	template<typename PointType>
+	void DrawViewFacingArc(const PointType& Position, float Radius, int Steps, float MinAngle, float MaxAngle, const FLinearColor& Color, float LineThicknessIn, bool bDepthTestedIn)
+	{
+		InternalDrawViewFacingArc((FVector)Position, Radius, Steps, MinAngle, MaxAngle, Color, LineThicknessIn, bDepthTestedIn);
+	}
+
 	/** Draw a world-space X facing the viewer at the given position. */
 	template<typename PointType>
 	void DrawViewFacingX(const PointType& Position, float Width)
@@ -263,6 +277,10 @@ protected:
 	INTERACTIVETOOLSFRAMEWORK_API virtual void InternalDrawCircle(const FVector& Position, const FVector& Normal, float Radius, int Steps, const FLinearColor& Color, float LineThickness, bool bDepthTested);
 	// actually does the circle drawing
 	INTERACTIVETOOLSFRAMEWORK_API virtual void InternalDrawViewFacingCircle(const FVector& Position, float Radius, int Steps, const FLinearColor& Color, float LineThickness, bool bDepthTested);
+	// actually does the arc drawing
+	INTERACTIVETOOLSFRAMEWORK_API virtual void InternalDrawArc(const FVector& InPosition, const FVector& InNormal, const float InRadius, int32 InSteps, const float InMinAngle, const float InMaxAngle, const FLinearColor& Color, float LineThickness, bool bDepthTested);
+	// actually does the arc drawing
+	INTERACTIVETOOLSFRAMEWORK_API virtual void InternalDrawViewFacingArc(const FVector& InPosition, float InRadius, int32 InSteps, const float InMinAngle, const float InMaxAngle, const FLinearColor& Color, float LineThickness, bool bDepthTested);
 	// actually does the cylinder drawing
 	INTERACTIVETOOLSFRAMEWORK_API virtual void InternalDrawWireCylinder(const FVector& Position, const FVector& Normal, float Radius, float Height, int Steps, const FLinearColor& Color, float LineThickness, bool bDepthTested);
 	// actually does the box drawing
