[ITF] UInputBehavior responding to app focus lost

UInputRouters now call ForceTerminateAll when the app loses focus.

- Added a virtual function called SupportsCapture to UInputBehavior. Returns false by default. Can be used to mark those behaviors which capture inputs even when not active.

- UKeyAsModifierInputBehavior has SupportsCapture returning true.

- UInputBehaviorSet provieds a method to retrieve a list of Behaviors supporting capture: TArray<UInputBehavior*> UInputBehaviorSet::GetCaptureSupportingBehaviors()

- The function is used by UInputRouter::ForceTerminateAll() to make sure also non active Behaviors with SupportsCapture returning true receive info about app focus change events

- New OnForceEndCapture virtual method fin IClickDragBehaviorTarget

diff --git a/Engine/Source/Runtime/InteractiveToolsFramework/InteractiveToolsFramework.Build.cs b/Engine/Source/Runtime/InteractiveToolsFramework/InteractiveToolsFramework.Build.cs
--- a/Engine/Source/Runtime/InteractiveToolsFramework/InteractiveToolsFramework.Build.cs
+++ b/Engine/Source/Runtime/InteractiveToolsFramework/InteractiveToolsFramework.Build.cs
@@ -26,8 +26,8 @@ public class InteractiveToolsFramework : ModuleRules
 				"Engine",
                 "RHI",
 				"GeometryCore",
-				"RenderCore" // for gizmo components
-				//"Slate",
+				"RenderCore", // for gizmo components
+				"Slate",
 				//"SlateCore",
 				// ... add private dependencies that you statically link with here ...	
 			}
diff --git a/Engine/Source/Runtime/InteractiveToolsFramework/Private/BaseBehaviors/KeyAsModifierInputBehavior.cpp b/Engine/Source/Runtime/InteractiveToolsFramework/Private/BaseBehaviors/KeyAsModifierInputBehavior.cpp
--- a/Engine/Source/Runtime/InteractiveToolsFramework/Private/BaseBehaviors/KeyAsModifierInputBehavior.cpp
+++ b/Engine/Source/Runtime/InteractiveToolsFramework/Private/BaseBehaviors/KeyAsModifierInputBehavior.cpp
@@ -40,9 +40,15 @@ FInputCaptureUpdate UKeyAsModifierInputBehavior::UpdateCapture(const FInputDevic
 	return FInputCaptureUpdate::End();
 }
 
+bool UKeyAsModifierInputBehavior::WantsForceEndCapture()
+{
+	return true;
+}
 
 void UKeyAsModifierInputBehavior::ForceEndCapture(const FInputCaptureData& Data)
 {
+	if (Target)
+	{
+		Target->OnForceEndCapture();
+	}
 }
-
-
diff --git a/Engine/Source/Runtime/InteractiveToolsFramework/Private/InputBehavior.cpp b/Engine/Source/Runtime/InteractiveToolsFramework/Private/InputBehavior.cpp
--- a/Engine/Source/Runtime/InteractiveToolsFramework/Private/InputBehavior.cpp
+++ b/Engine/Source/Runtime/InteractiveToolsFramework/Private/InputBehavior.cpp
@@ -89,4 +89,7 @@ void UInputBehavior::EndHoverCapture()
 	return;
 }
 
-
+bool UInputBehavior::WantsForceEndCapture()
+{
+	return false;
+}
diff --git a/Engine/Source/Runtime/InteractiveToolsFramework/Private/InputBehaviorSet.cpp b/Engine/Source/Runtime/InteractiveToolsFramework/Private/InputBehaviorSet.cpp
--- a/Engine/Source/Runtime/InteractiveToolsFramework/Private/InputBehaviorSet.cpp
+++ b/Engine/Source/Runtime/InteractiveToolsFramework/Private/InputBehaviorSet.cpp
@@ -135,9 +135,19 @@ void UInputBehaviorSet::CollectWantsHoverCapture(const FInputDeviceState& input,
 	}
 }
 
-
-
-
+void UInputBehaviorSet::CollectWantsForceEndCapture(TArray<UInputBehavior*>& ResultOut)
+{
+	for (const FBehaviorInfo& BehaviorInfo : Behaviors)
+	{
+		if (UInputBehavior* const Behavior = BehaviorInfo.Behavior)
+		{
+			if (Behavior->WantsForceEndCapture())
+			{
+				ResultOut.Add(Behavior);
+			}
+		}
+	}
+}
 
 void UInputBehaviorSet::BehaviorsModified()
 {
diff --git a/Engine/Source/Runtime/InteractiveToolsFramework/Private/InputRouter.cpp b/Engine/Source/Runtime/InteractiveToolsFramework/Private/InputRouter.cpp
--- a/Engine/Source/Runtime/InteractiveToolsFramework/Private/InputRouter.cpp
+++ b/Engine/Source/Runtime/InteractiveToolsFramework/Private/InputRouter.cpp
@@ -2,6 +2,7 @@
 
 
 #include "InputRouter.h"
+#include "Framework/Application/SlateApplication.h"
 
 #include UE_INLINE_GENERATED_CPP_BY_NAME(InputRouter)
 
@@ -21,6 +22,9 @@ UInputRouter::UInputRouter()
 	bAutoInvalidateOnCapture = false;
 
 	ActiveInputBehaviors = NewObject<UInputBehaviorSet>(this, "InputBehaviors");
+
+	// Terminate all captures when editor loses focus
+	FSlateApplication::Get().OnApplicationActivationStateChanged().AddUObject(this, &UInputRouter::OnApplicationFocusChanged);
 }
 
 
@@ -322,6 +326,14 @@ bool UInputRouter::UpdateExistingHoverCaptureIfPresent(const FInputDeviceState&
 	return false;
 }
 
+void UInputRouter::OnApplicationFocusChanged(bool bInIsFocused)
+{
+	if (!bInIsFocused)
+	{
+		ForceTerminateAll();
+	}
+}
+
 bool UInputRouter::ProcessMouseHover(const FInputDeviceState& Input)
 {
 	TArray<FInputCaptureRequest> CaptureRequests;
@@ -409,6 +421,13 @@ void UInputRouter::ForceTerminateAll()
 	{
 		TerminateHover(EInputCaptureSide::Left);
 	}
+
+	TArray<UInputBehavior*> ForceEndCaptureBehaviors;
+	ActiveInputBehaviors->CollectWantsForceEndCapture(ForceEndCaptureBehaviors);
+	for (UInputBehavior* const Behavior : ForceEndCaptureBehaviors)
+	{
+		Behavior->ForceEndCapture(FInputCaptureData());
+	}
 }
 
 
diff --git a/Engine/Source/Runtime/InteractiveToolsFramework/Public/BaseBehaviors/BehaviorTargetInterfaces.h b/Engine/Source/Runtime/InteractiveToolsFramework/Public/BaseBehaviors/BehaviorTargetInterfaces.h
--- a/Engine/Source/Runtime/InteractiveToolsFramework/Public/BaseBehaviors/BehaviorTargetInterfaces.h
+++ b/Engine/Source/Runtime/InteractiveToolsFramework/Public/BaseBehaviors/BehaviorTargetInterfaces.h
@@ -24,6 +24,13 @@ public:
 	virtual void OnUpdateModifierState(int ModifierID, bool bIsOn)
 	{
 	}
+
+	/**
+	 * Notify Target that capture is ending
+	 */
+	virtual void OnForceEndCapture()
+	{
+	}
 };
 
 /**
@@ -272,4 +279,4 @@ public:
 	
 	virtual FKey GetCaptureKey() { return EKeys::B; }
 	virtual void PostDragUpdated() {}
-};
\ No newline at end of file
+};
diff --git a/Engine/Source/Runtime/InteractiveToolsFramework/Public/BaseBehaviors/KeyAsModifierInputBehavior.h b/Engine/Source/Runtime/InteractiveToolsFramework/Public/BaseBehaviors/KeyAsModifierInputBehavior.h
--- a/Engine/Source/Runtime/InteractiveToolsFramework/Public/BaseBehaviors/KeyAsModifierInputBehavior.h
+++ b/Engine/Source/Runtime/InteractiveToolsFramework/Public/BaseBehaviors/KeyAsModifierInputBehavior.h
@@ -59,6 +59,7 @@ public:
 	INTERACTIVETOOLSFRAMEWORK_API virtual FInputCaptureRequest WantsCapture(const FInputDeviceState& Input) override;
 	INTERACTIVETOOLSFRAMEWORK_API virtual FInputCaptureUpdate BeginCapture(const FInputDeviceState& Input, EInputCaptureSide Side) override;
 	INTERACTIVETOOLSFRAMEWORK_API virtual FInputCaptureUpdate UpdateCapture(const FInputDeviceState& Input, const FInputCaptureData& Data) override;
+	INTERACTIVETOOLSFRAMEWORK_API virtual bool WantsForceEndCapture() override;
 	INTERACTIVETOOLSFRAMEWORK_API virtual void ForceEndCapture(const FInputCaptureData& Data) override;
 
 
diff --git a/Engine/Source/Runtime/InteractiveToolsFramework/Public/InputBehavior.h b/Engine/Source/Runtime/InteractiveToolsFramework/Public/InputBehavior.h
--- a/Engine/Source/Runtime/InteractiveToolsFramework/Public/InputBehavior.h
+++ b/Engine/Source/Runtime/InteractiveToolsFramework/Public/InputBehavior.h
@@ -11,8 +11,6 @@
 
 #include "InputBehavior.generated.h"
 
-
-
 /**
  * Input can be captured separately for Left and Right sides (eg for VR controllers)
  * Currently mouse is Left.
@@ -27,7 +25,6 @@ enum class EInputCaptureSide
 	Any = 99
 };
 
-
 /**
  * An active capturing behavior may need to keep track of additional data that
  * cannot be stored within the behavior (for example if the same behavior instance
@@ -48,7 +45,6 @@ struct FInputCaptureData
 	}
 };
 
-
 /**
  * Used by FInputCaptureRequest to indicate whether the InputBehavior
  * wants to capture or ignore an input event
@@ -60,11 +56,9 @@ enum class EInputCaptureRequestType
 	Ignore = 2
 };
 
-
 // predeclaration
 class UInputBehavior;
 
-
 /**
  * UInputBehavior returns an FInputCaptureRequest from WantsCapture() to indicate
  * whether it wants to capture or ignore an input event
@@ -73,7 +67,7 @@ struct FInputCaptureRequest
 {
 	/** Which input behavior generated this request */
 	UInputBehavior* Source;
-    /** What type of capture request is this (Begin or Ignore) */
+	/** What type of capture request is this (Begin or Ignore) */
 	EInputCaptureRequestType Type;
 	/** Which side does request want to capture on */
 	EInputCaptureSide Side;
@@ -84,8 +78,12 @@ struct FInputCaptureRequest
 	/** Owner of the requesting behavior. Behavior doesn't know this, so this is initialized to null */
 	void* Owner;
 
-
-	FInputCaptureRequest(EInputCaptureRequestType type, UInputBehavior* behavior, EInputCaptureSide whichSide, double hitDepth = TNumericLimits<double>::Max() )
+	FInputCaptureRequest(
+		EInputCaptureRequestType type,
+		UInputBehavior* behavior,
+		EInputCaptureSide whichSide,
+		double hitDepth = TNumericLimits<double>::Max()
+	)
 	{
 		this->Type = type;
 		this->Source = behavior;
@@ -95,25 +93,24 @@ struct FInputCaptureRequest
 	}
 
 	/** Create a Begin-capture request */
-	static FInputCaptureRequest Begin(UInputBehavior* behavior, EInputCaptureSide whichSide, double hitDepth = TNumericLimits<double>::Max() )
+	static FInputCaptureRequest Begin(
+		UInputBehavior* behavior, EInputCaptureSide whichSide, double hitDepth = TNumericLimits<double>::Max()
+	)
 	{
 		return FInputCaptureRequest(EInputCaptureRequestType::Begin, behavior, whichSide, hitDepth);
 	}
 
 	/** Create an ignore-capture request */
-	static FInputCaptureRequest Ignore() 
+	static FInputCaptureRequest Ignore()
 	{
-		return FInputCaptureRequest(EInputCaptureRequestType::Ignore, nullptr, EInputCaptureSide::Any, TNumericLimits<double>::Max() );
+		return FInputCaptureRequest(
+			EInputCaptureRequestType::Ignore, nullptr, EInputCaptureSide::Any, TNumericLimits<double>::Max()
+		);
 	}
 
 	friend bool operator<(const FInputCaptureRequest& l, const FInputCaptureRequest& r);
 };
 
-
-
-
-
-
 /**
  * FInputCaptureUpdate uses this type to indicate what state the capturing Behavior
  * would like to transition to, based on the input event
@@ -121,14 +118,12 @@ struct FInputCaptureRequest
 UENUM()
 enum class EInputCaptureState
 {
-	Begin = 1,		// start capturing (which should always be the case if BeginCapture is called)
-	Continue = 2,   // Behavior wants to continue capturing
-	End = 3,        // Behavior wants to end capturing
-	Ignore = 4      // Behavior ignored this event
+	Begin = 1, // start capturing (which should always be the case if BeginCapture is called)
+	Continue = 2, // Behavior wants to continue capturing
+	End = 3, // Behavior wants to end capturing
+	Ignore = 4 // Behavior ignored this event
 };
 
-
-
 /**
  * IInputBehavior returns an FInputCaptureUpdate from BeginCapture() and UpdateCapture(),
  * which indicates to the InputRouter what the Behavior would like to have happen.
@@ -142,7 +137,7 @@ struct FInputCaptureUpdate
 	/** custom data for the active capture that should be propagated to next UpdateCapture() call */
 	FInputCaptureData Data;
 
-	/** 
+	/**
 	 * Create a begin-capturing instance of FInputCaptureUpdate
 	 * @param Source UInputBehavior that is returning this update
 	 * @param Which Which side we are capturing on
@@ -175,7 +170,9 @@ struct FInputCaptureUpdate
 	 * @param Which Which side we are capturing on
 	 * @param CustomData client-provided data that will be passed to UInputBehavior::UpdateCapture() calls. Client owns this memory!
 	 */
-	FInputCaptureUpdate(EInputCaptureState StateIn, UInputBehavior* SourceBehaviorIn, EInputCaptureSide WhichSideIn, void* CustomData = nullptr)
+	FInputCaptureUpdate(
+		EInputCaptureState StateIn, UInputBehavior* SourceBehaviorIn, EInputCaptureSide WhichSideIn, void* CustomData = nullptr
+	)
 	{
 		State = StateIn;
 		Source = SourceBehaviorIn;
@@ -184,8 +181,6 @@ struct FInputCaptureUpdate
 	}
 };
 
-
-
 /**
  * Each UInputBehavior provides a priority that is used to help resolve situations
  * when multiple Behaviors want to capture based on the same input event
@@ -225,10 +220,8 @@ struct FInputCapturePriority
 	}
 };
 
-
-
 /**
- * An InputBehavior implements a state machine for a user interaction. 
+ * An InputBehavior implements a state machine for a user interaction.
  * The InputRouter maintains a set of active Behaviors, and when new input
  * events occur, it calls WantsCapture() to check if the Behavior would like to
  * begin capturing the applicable input event stream (eg for a mouse, one or both VR controllers, etc).
@@ -246,8 +239,8 @@ struct FInputCapturePriority
  *            - return End
  *
  * Written sufficiently generically, the above Behavior doesn't need to know about buttons,
- * it just needs to know how to hit-test the clickable object(s). Similarly separate 
- * Behaviors can be written for mouse, VR, touch, gamepad, etc. 
+ * it just needs to know how to hit-test the clickable object(s). Similarly separate
+ * Behaviors can be written for mouse, VR, touch, gamepad, etc.
  *
  * Implementing interactions in this way allows the input handling to be separated from functionality.
  */
@@ -265,7 +258,6 @@ public:
 	/** Configure the default priority of an instance of this behavior */
 	INTERACTIVETOOLSFRAMEWORK_API virtual void SetDefaultPriority(const FInputCapturePriority& Priority);
 
-
 	/** Which device types does this Behavior support */
 	INTERACTIVETOOLSFRAMEWORK_API virtual EInputDevices GetSupportedDevices();
 
@@ -273,18 +265,21 @@ public:
 	INTERACTIVETOOLSFRAMEWORK_API virtual FInputCaptureRequest WantsCapture(const FInputDeviceState& InputState);
 
 	/** Called after WantsCapture() returns a capture request that was accepted */
-	INTERACTIVETOOLSFRAMEWORK_API virtual FInputCaptureUpdate BeginCapture(const FInputDeviceState& InputState, EInputCaptureSide eSide);
+	INTERACTIVETOOLSFRAMEWORK_API virtual FInputCaptureUpdate BeginCapture(
+		const FInputDeviceState& InputState, EInputCaptureSide eSide
+	);
 
-	/** 
+	/**
 	 * Called for each new input event during a capture sequence. Return Continue to keep
 	 * capturing, or End to finish capturing.
 	 */
-	INTERACTIVETOOLSFRAMEWORK_API virtual FInputCaptureUpdate UpdateCapture(const FInputDeviceState& InputState, const FInputCaptureData& CaptureData);
+	INTERACTIVETOOLSFRAMEWORK_API virtual FInputCaptureUpdate UpdateCapture(
+		const FInputDeviceState& InputState, const FInputCaptureData& CaptureData
+	);
 
 	/** If this is called, the Behavior has forcibly lost capture (eg due to app losing focus for example) and needs to clean up accordingly */
 	INTERACTIVETOOLSFRAMEWORK_API virtual void ForceEndCapture(const FInputCaptureData& CaptureData);
 
-
 	//
 	// hover support (optional)
 	//
@@ -296,7 +291,9 @@ public:
 	INTERACTIVETOOLSFRAMEWORK_API virtual FInputCaptureRequest WantsHoverCapture(const FInputDeviceState& InputState);
 
 	/** Called after WantsHoverCapture() returns a capture request that was accepted */
-	INTERACTIVETOOLSFRAMEWORK_API virtual FInputCaptureUpdate BeginHoverCapture(const FInputDeviceState& InputState, EInputCaptureSide eSide);
+	INTERACTIVETOOLSFRAMEWORK_API virtual FInputCaptureUpdate BeginHoverCapture(
+		const FInputDeviceState& InputState, EInputCaptureSide eSide
+	);
 
 	/** Called on each new hover input event, ie if no other behavior is actively capturing input */
 	INTERACTIVETOOLSFRAMEWORK_API virtual FInputCaptureUpdate UpdateHoverCapture(const FInputDeviceState& InputState);
@@ -304,6 +301,11 @@ public:
 	/** If a different hover capture begins, focus is lost, a tool starts, etc, any active hover visualization needs to terminate */
 	INTERACTIVETOOLSFRAMEWORK_API virtual void EndHoverCapture();
 
+	/**
+	 * When this function returns true, this behavior will always receive ForceEndCapture calls even when it is not actively
+	 * capturing. This is useful for behaviors that do not capture, but hold state (ex. Key Modifiers) and want to be notified when to clear that state.
+	 */
+	INTERACTIVETOOLSFRAMEWORK_API virtual bool WantsForceEndCapture();
 
 protected:
 	/** priority returned by GetPriority() */
diff --git a/Engine/Source/Runtime/InteractiveToolsFramework/Public/InputBehaviorSet.h b/Engine/Source/Runtime/InteractiveToolsFramework/Public/InputBehaviorSet.h
--- a/Engine/Source/Runtime/InteractiveToolsFramework/Public/InputBehaviorSet.h
+++ b/Engine/Source/Runtime/InteractiveToolsFramework/Public/InputBehaviorSet.h
@@ -130,6 +130,10 @@ public:
 	 */
 	INTERACTIVETOOLSFRAMEWORK_API virtual void CollectWantsHoverCapture(const FInputDeviceState& InputState, TArray<FInputCaptureRequest>& ResultOut);
 
+	/**
+	 * Call UInputBehavior::WantsForceEndCapture on each valid Behavior and collect up the requests that indicated that a ForceEndCapture was desired.
+	 */
+	INTERACTIVETOOLSFRAMEWORK_API virtual void CollectWantsForceEndCapture(TArray<UInputBehavior*>& ResultOut);
 
 
 protected:
diff --git a/Engine/Source/Runtime/InteractiveToolsFramework/Public/InputRouter.h b/Engine/Source/Runtime/InteractiveToolsFramework/Public/InputRouter.h
--- a/Engine/Source/Runtime/InteractiveToolsFramework/Public/InputRouter.h
+++ b/Engine/Source/Runtime/InteractiveToolsFramework/Public/InputRouter.h
@@ -123,4 +123,5 @@ protected:
 	INTERACTIVETOOLSFRAMEWORK_API void TerminateHover(EInputCaptureSide Side);
 	INTERACTIVETOOLSFRAMEWORK_API bool ProcessMouseHover(const FInputDeviceState& Input);
 	INTERACTIVETOOLSFRAMEWORK_API bool UpdateExistingHoverCaptureIfPresent(const FInputDeviceState& Input);
+	INTERACTIVETOOLSFRAMEWORK_API void OnApplicationFocusChanged(bool bInIsFocused);
 };
