[ITF] Fix issue with UInputRouter accessing SlateApplication when not available

Adding validity check for FSlateApplication in UInputRouter constructor.

diff --git a/Engine/Source/Runtime/InteractiveToolsFramework/Private/InputRouter.cpp b/Engine/Source/Runtime/InteractiveToolsFramework/Private/InputRouter.cpp
--- a/Engine/Source/Runtime/InteractiveToolsFramework/Private/InputRouter.cpp
+++ b/Engine/Source/Runtime/InteractiveToolsFramework/Private/InputRouter.cpp
@@ -23,8 +23,11 @@ UInputRouter::UInputRouter()
 
 	ActiveInputBehaviors = NewObject<UInputBehaviorSet>(this, "InputBehaviors");
 
-	// Terminate all captures when editor loses focus
-	FSlateApplication::Get().OnApplicationActivationStateChanged().AddUObject(this, &UInputRouter::OnApplicationFocusChanged);
+	if (FSlateApplication::IsInitialized())
+	{
+		// Terminate all captures when editor loses focus
+		FSlateApplication::Get().OnApplicationActivationStateChanged().AddUObject(this, &UInputRouter::OnApplicationFocusChanged);
+	}
 }
 
 
