[ITF] Added EndCaps option to UGizmoElementCone + UGizmoElementArrow

diff --git a/Engine/Source/Runtime/InteractiveToolsFramework/Private/BaseGizmos/GizmoElementArrow.cpp b/Engine/Source/Runtime/InteractiveToolsFramework/Private/BaseGizmos/GizmoElementArrow.cpp
--- a/Engine/Source/Runtime/InteractiveToolsFramework/Private/BaseGizmos/GizmoElementArrow.cpp
+++ b/Engine/Source/Runtime/InteractiveToolsFramework/Private/BaseGizmos/GizmoElementArrow.cpp
@@ -239,6 +239,20 @@ int32 UGizmoElementArrow::GetNumSides() const
 	return NumSides;
 }
 
+void UGizmoElementArrow::SetEndCaps(bool InEndCaps)
+{
+	if (bEndCaps != InEndCaps)
+	{
+		bEndCaps = InEndCaps;
+		bUpdateArrowHead = true;
+	}
+}
+
+bool UGizmoElementArrow::GetEndCaps() const
+{
+	return bEndCaps;
+}
+
 void UGizmoElementArrow::SetPixelHitDistanceThreshold(float InPixelHitDistanceThreshold)
 {
 	if (PixelHitDistanceThreshold != InPixelHitDistanceThreshold)
@@ -304,6 +318,7 @@ void UGizmoElementArrow::UpdateArrowHead()
 		ConeElement->SetNumSides(NumSides);
 		ConeElement->SetElementInteractionState(ElementInteractionState);
 		ConeElement->SetPixelHitDistanceThreshold(PixelHitDistanceThreshold);
+		ConeElement->SetEndCaps(bEndCaps);
 	}
 	else // (HeadType == EGizmoElementArrowHeadType::Cube)
 	{
diff --git a/Engine/Source/Runtime/InteractiveToolsFramework/Private/BaseGizmos/GizmoElementCone.cpp b/Engine/Source/Runtime/InteractiveToolsFramework/Private/BaseGizmos/GizmoElementCone.cpp
--- a/Engine/Source/Runtime/InteractiveToolsFramework/Private/BaseGizmos/GizmoElementCone.cpp
+++ b/Engine/Source/Runtime/InteractiveToolsFramework/Private/BaseGizmos/GizmoElementCone.cpp
@@ -5,10 +5,62 @@
 #include "GizmoPrivateUtil.h"
 #include "BaseGizmos/GizmoRenderingUtil.h"
 #include "BaseGizmos/GizmoMath.h"
+#include "DynamicMeshBuilder.h"
 #include "InputState.h"
 #include "Materials/MaterialInterface.h"
 #include "PrimitiveDrawingUtils.h"
 
+namespace GizmoConeComponentLocals
+{
+	// Custom DrawDisc with flat/constant normal
+	void DrawDisc(FPrimitiveDrawInterface* InPDI, const FVector& InBase, const FVector& InAxisX, const FVector& InAxisY, double InRadius, int32 InNumSides, const UMaterialInterface* InMaterial)
+	{
+		check(InNumSides >= 3);
+
+		const float	AngleDelta = 2.0f * UE_PI / InNumSides;
+
+		const FVector2f TextureCoordinate = FVector2f(0.0f, 0.0f);
+		const FVector2f::FReal TextureCoordinateStep = 1.0f / InNumSides;
+
+		const FVector AxisZ = (InAxisX) ^ InAxisY;
+
+		FDynamicMeshBuilder MeshBuilder(InPDI->View->GetFeatureLevel());
+
+		// Compute vertices for base circle.
+		for (int32 SideIndex = 0; SideIndex < InNumSides; ++SideIndex)
+		{
+			const FVector Vertex = InBase + (InAxisX * FMath::Cos(AngleDelta * (SideIndex)) + InAxisY * FMath::Sin(AngleDelta * (SideIndex))) * InRadius;
+
+			FDynamicMeshVertex MeshVertex;
+			MeshVertex.Position = static_cast<FVector3f>(Vertex);
+			MeshVertex.Color = FColor::White;
+			MeshVertex.TextureCoordinate[0] = FVector2f(TextureCoordinate);
+			MeshVertex.TextureCoordinate[0].X += TextureCoordinateStep * SideIndex;
+
+			MeshVertex.SetTangents(
+			FVector3f(InAxisY),
+			FVector3f(-InAxisX),
+			FVector3f(AxisZ)
+			);
+
+			MeshBuilder.AddVertex(MeshVertex); //Add bottom vertex
+		}
+
+		// Add top/bottom triangles, in the style of a fan.
+		for (int32 SideIndex = 0; SideIndex < InNumSides - 1; ++SideIndex)
+		{
+			constexpr int32 V0 = 0;
+			const int32 V1 = SideIndex;
+			const int32 V2 = (SideIndex + 1);
+
+			MeshBuilder.AddTriangle(V0, V1, V2);
+			MeshBuilder.AddTriangle(V0, V2, V1);
+		}
+
+		MeshBuilder.Draw(InPDI, FMatrix::Identity, InMaterial->GetRenderProxy(), SDPG_Foreground,0.0f);
+	}
+}
+
 #include UE_INLINE_GENERATED_CPP_BY_NAME(GizmoElementCone)
 
 void UGizmoElementCone::Render(IToolsContextRenderAPI* RenderAPI, const FRenderTraversalState& RenderState)
@@ -25,9 +77,20 @@ void UGizmoElementCone::Render(IToolsContextRenderAPI* RenderAPI, const FRenderT
 
 			FTransform RenderLocalToWorldTransform = FTransform(Rotation, FVector::ZeroVector, Scale) * CurrentRenderState.LocalToWorldTransform;
 			const float ConeSide = FMath::Sqrt(Height * Height + Radius * Radius);
-			const float Angle = FMath::Acos(Height / ConeSide);
+			const float HeightOverConeSide = Height / ConeSide;
+			const float Angle = FMath::Acos(HeightOverConeSide);
 			FPrimitiveDrawInterface* PDI = RenderAPI->GetPrimitiveDrawInterface();
 			DrawCone(PDI, RenderLocalToWorldTransform.ToMatrixWithScale(), Angle, Angle, NumSides, false, FColor::White, UseMaterial->GetRenderProxy(), SDPG_Foreground);
+
+			if (bEndCaps)
+			{
+				// The resulting length of the arrow/cone is altered according to the angle, so we need to account for it here to get an accurate bottom location
+				const float Offset = Height * HeightOverConeSide;
+				const FVector Location = CurrentRenderState.LocalToWorldTransform.TransformPosition(Direction * Offset);
+				const FVector XAxis = RenderLocalToWorldTransform.TransformVectorNoScale(FVector::RightVector);
+				const FVector YAxis = RenderLocalToWorldTransform.TransformVectorNoScale(FVector::UpVector);
+				GizmoConeComponentLocals::DrawDisc(PDI, Location, XAxis, YAxis, Radius * HeightOverConeSide * CurrentRenderState.LocalToWorldTransform.GetScale3D().X, NumSides, UseMaterial);
+			}
 		}
 	}
 }
@@ -127,4 +190,12 @@ int32 UGizmoElementCone::GetNumSides() const
 	return NumSides;
 }
 
+void UGizmoElementCone::SetEndCaps(bool InEndCaps)
+{
+	bEndCaps = InEndCaps;
+}
 
+bool UGizmoElementCone::GetEndCaps() const
+{
+	return bEndCaps;
+}
diff --git a/Engine/Source/Runtime/InteractiveToolsFramework/Public/BaseGizmos/GizmoElementArrow.h b/Engine/Source/Runtime/InteractiveToolsFramework/Public/BaseGizmos/GizmoElementArrow.h
--- a/Engine/Source/Runtime/InteractiveToolsFramework/Public/BaseGizmos/GizmoElementArrow.h
+++ b/Engine/Source/Runtime/InteractiveToolsFramework/Public/BaseGizmos/GizmoElementArrow.h
@@ -69,6 +69,10 @@ public:
 	// Number of sides for cylinder and cone, if relevant.
 	INTERACTIVETOOLSFRAMEWORK_API virtual void SetNumSides(int32 InNumSides);
 	INTERACTIVETOOLSFRAMEWORK_API virtual int32 GetNumSides() const;
+	
+	// Whether to render and end cap on the base of the arrow head. Applicable when the ArrowHeadType is Cone.
+	INTERACTIVETOOLSFRAMEWORK_API virtual void SetEndCaps(bool InEndCaps);
+	INTERACTIVETOOLSFRAMEWORK_API virtual bool GetEndCaps() const;
 
 	// Head type cone or cube.
 	INTERACTIVETOOLSFRAMEWORK_API virtual void SetHeadType(EGizmoElementArrowHeadType InHeadType);
@@ -138,6 +142,10 @@ protected:
 	UPROPERTY()
 	int32 NumSides = 32;
 
+	// Whether to render and end cap on the base of the arrow head.
+	UPROPERTY()
+	bool bEndCaps = false;
+
 	// Head type
 	UPROPERTY()
 	EGizmoElementArrowHeadType HeadType = EGizmoElementArrowHeadType::Cone;
diff --git a/Engine/Source/Runtime/InteractiveToolsFramework/Public/BaseGizmos/GizmoElementCone.h b/Engine/Source/Runtime/InteractiveToolsFramework/Public/BaseGizmos/GizmoElementCone.h
--- a/Engine/Source/Runtime/InteractiveToolsFramework/Public/BaseGizmos/GizmoElementCone.h
+++ b/Engine/Source/Runtime/InteractiveToolsFramework/Public/BaseGizmos/GizmoElementCone.h
@@ -45,6 +45,10 @@ public:
 	INTERACTIVETOOLSFRAMEWORK_API virtual void SetNumSides(int32 InNumSides);
 	INTERACTIVETOOLSFRAMEWORK_API virtual int32 GetNumSides() const;
 
+	// Whether to render and end cap on the base of the cone.
+	INTERACTIVETOOLSFRAMEWORK_API virtual void SetEndCaps(bool InEndCaps);
+	INTERACTIVETOOLSFRAMEWORK_API virtual bool GetEndCaps() const;
+
 protected:
 
 	// Cone tip location.
@@ -66,5 +70,8 @@ protected:
 	// Number of sides for tessellating cone.
 	UPROPERTY()
 	int32 NumSides = 32;
+	
+	// Whether to render and end cap on the base of the cone.
+	UPROPERTY()
+	bool bEndCaps = false;
 };
-
