Combined Gizmo Transform will display LeftUpForward if configured
* FSubGizmoTransformAdjuster will adjust gizmo component transforms if view mirroring is not enabled to display the Y leg pointing to left
*  CombinedTransformGizmo computes own axis over which transform adjustments are performed against.  This was changed to take into account the LeftUpForward configuration so input around the Right-Left axis is converted back to engine space.

diff --git a/Engine/Source/Runtime/InteractiveToolsFramework/Private/BaseGizmos/TransformSubGizmoUtil.cpp b/Engine/Source/Runtime/InteractiveToolsFramework/Private/BaseGizmos/TransformSubGizmoUtil.cpp
--- a/Engine/Source/Runtime/InteractiveToolsFramework/Private/BaseGizmos/TransformSubGizmoUtil.cpp
+++ b/Engine/Source/Runtime/InteractiveToolsFramework/Private/BaseGizmos/TransformSubGizmoUtil.cpp
@@ -8,23 +8,31 @@
 #include "BaseGizmos/ViewBasedTransformAdjusters.h"
 #include "BaseGizmos/ViewAdjustedStaticMeshGizmoComponent.h"
 #include "Components/SceneComponent.h"
+#include "Misc/AxisDisplayInfo.h"
 #include "Misc/EnumRange.h" // MakeFlagsRange
 
 
 FTransform UE::GizmoUtil::GetRotatedBasisTransform(const FTransform& TransformIn, EAxis::Type AxisToBeX)
 {
+	FTransform Result = TransformIn;
 	if (AxisToBeX == EAxis::Y)
 	{
 		FQuat RotationQuat(FMatrix(FVector::YAxisVector, FVector::ZAxisVector, FVector::XAxisVector, FVector::ZeroVector));
-		return TransformIn * RotationQuat;
+		Result = TransformIn * RotationQuat;
+		
 	}
 	else if (AxisToBeX == EAxis::Z)
 	{
 		FQuat RotationQuat(FMatrix(FVector::ZAxisVector, FVector::XAxisVector, FVector::YAxisVector, FVector::ZeroVector));
-		return TransformIn * RotationQuat;
+		Result = TransformIn * RotationQuat;
 	}
 
-	return TransformIn;
+	if (AxisDisplayInfo::GetAxisDisplayCoordinateSystem() == EAxisList::LeftUpForward)
+	{
+		Result.Mirror(EAxis::Y, EAxis::None);
+	}
+
+	return Result;
 }
 
 void UE::GizmoUtil::InitializeSubGizmoElementsWithMesh(const FInitMeshSubGizmoParams& Params,
diff --git a/Engine/Source/Runtime/InteractiveToolsFramework/Private/BaseGizmos/ViewBasedTransformAdjusters.cpp b/Engine/Source/Runtime/InteractiveToolsFramework/Private/BaseGizmos/ViewBasedTransformAdjusters.cpp
--- a/Engine/Source/Runtime/InteractiveToolsFramework/Private/BaseGizmos/ViewBasedTransformAdjusters.cpp
+++ b/Engine/Source/Runtime/InteractiveToolsFramework/Private/BaseGizmos/ViewBasedTransformAdjusters.cpp
@@ -5,6 +5,7 @@
 #include "BaseGizmos/GizmoRenderingUtil.h"
 #include "BaseGizmos/ViewAdjustedStaticMeshGizmoComponent.h"
 #include "FrameTypes.h"
+#include "Misc/AxisDisplayInfo.h"
 
 namespace ViewBasedTransformAdjustersLocals
 {
@@ -78,6 +79,7 @@ namespace ViewBasedTransformAdjustersLocals
 
 		if (Settings.bMirrorBasedOnOctant)
 		{
+			// Coordinate axis are mirrored based on view octant to allow for easier access to all gizmo controls
 			FVector3d GizmoSpaceDirectionTowardCamera;
 			if (View.IsPerspectiveProjection())
 			{
@@ -96,6 +98,12 @@ namespace ViewBasedTransformAdjustersLocals
 					ComponentToGizmo.Mirror(Axes[Dim], EAxis::None);
 				}
 			}
+			if (AxisDisplayInfo::GetAxisDisplayCoordinateSystem() == EAxisList::LeftUpForward)
+			{
+				// LeftUpForward coordinate system display mirrors Unreal's engine Y so that axis is pointed to the left
+				// This effectively undoes the mirroring that occurs in GetRotatedBaseTransform
+				ComponentToGizmo.Mirror(EAxis::Y, EAxis::None);
+			}
 		}
 
 		if (Settings.bKeepConstantViewSize)
