InteractiveToolsFramework: Make the components used by gizmos not be hit by world line traces, so that gizmos don't interrupt raycasts in scriptable tools and elsewhere.

diff --git a/Engine/Source/Runtime/InteractiveToolsFramework/Private/BaseGizmos/GizmoRenderingUtil.cpp b/Engine/Source/Runtime/InteractiveToolsFramework/Private/BaseGizmos/GizmoRenderingUtil.cpp
--- a/Engine/Source/Runtime/InteractiveToolsFramework/Private/BaseGizmos/GizmoRenderingUtil.cpp
+++ b/Engine/Source/Runtime/InteractiveToolsFramework/Private/BaseGizmos/GizmoRenderingUtil.cpp
@@ -83,6 +83,15 @@ UViewAdjustedStaticMeshGizmoComponent* UE::GizmoRenderingUtil::CreateDefaultMate
 	// Not sure that this actually gets respected in any way for non-PDI calls, but just in case
 	Component->DepthPriorityGroup = SDPG_Foreground;
 
+	Component->SetCollisionEnabled(ECollisionEnabled::QueryOnly);
+	// Important: make it so that typical world line traces don't hit this gizmo component. 
+	Component->SetCollisionResponseToAllChannels(ECR_Ignore);
+	// The LineTraceComponent calls that UGizmoComponentHitTarget uses work even without us being bound
+	//  to a particular collision channel (so the following two lines are not strictly speaking necessary),
+	//  but let's go ahead and associate ourselves with the COLLISION_GIZMO channel (modeled on UGizmoHandleGroup::CreateMeshHandle).
+	Component->SetCollisionResponseToChannel(COLLISION_GIZMO, ECollisionResponse::ECR_Block);
+	Component->SetCollisionObjectType(COLLISION_GIZMO);
+
 	Component->SetAllMaterials(UE::GizmoRenderingUtil::GetDefaultGizmoComponentMaterial(Color, Component));
 	if (bAddHoverMaterial)
 	{
