InteractiveToolsFramework: Fix the rotation sub gizmo substitute component (that shows during interaction) not being properly aligned in global mode.

diff --git a/Engine/Source/Runtime/InteractiveToolsFramework/Private/BaseGizmos/CombinedTransformGizmo.cpp b/Engine/Source/Runtime/InteractiveToolsFramework/Private/BaseGizmos/CombinedTransformGizmo.cpp
--- a/Engine/Source/Runtime/InteractiveToolsFramework/Private/BaseGizmos/CombinedTransformGizmo.cpp
+++ b/Engine/Source/Runtime/InteractiveToolsFramework/Private/BaseGizmos/CombinedTransformGizmo.cpp
@@ -505,8 +505,7 @@ ACombinedTransformGizmoActor* ACombinedTransformGizmoActor::ConstructCustom3Axis
 
 						UE::GizmoRenderingUtil::FSubGizmoTransformAdjuster::AddTransformAdjuster(
 							SubstituteComponent,
-							// Don't need to keep distance to gizmo because we should be at gizmo root
-							nullptr, 
+							NewActor->GetRootComponent(),
 							/*bMirror*/ false);
 					}
 				}
diff --git a/Engine/Source/Runtime/InteractiveToolsFramework/Private/BaseGizmos/ViewBasedTransformAdjusters.cpp b/Engine/Source/Runtime/InteractiveToolsFramework/Private/BaseGizmos/ViewBasedTransformAdjusters.cpp
--- a/Engine/Source/Runtime/InteractiveToolsFramework/Private/BaseGizmos/ViewBasedTransformAdjusters.cpp
+++ b/Engine/Source/Runtime/InteractiveToolsFramework/Private/BaseGizmos/ViewBasedTransformAdjusters.cpp
@@ -153,7 +153,7 @@ void UE::GizmoRenderingUtil::FSubGizmoTransformAdjuster::SetGizmoOriginTransform
 }
 
 TSharedPtr<UE::GizmoRenderingUtil::FSubGizmoTransformAdjuster> UE::GizmoRenderingUtil::FSubGizmoTransformAdjuster::AddTransformAdjuster(
-	UViewAdjustedStaticMeshGizmoComponent* ComponentIn, USceneComponent* ComponentToKeepDistanceConstantTo, bool bMirrorBasedOnOctant)
+	UViewAdjustedStaticMeshGizmoComponent* ComponentIn, USceneComponent* GizmoRootComponent, bool bMirrorBasedOnOctant)
 {
 	if (!ComponentIn)
 	{
@@ -171,10 +171,10 @@ TSharedPtr<UE::GizmoRenderingUtil::FSubGizmoTransformAdjuster> UE::GizmoRenderin
 	Settings.DistanceInFrontOfCameraInOrtho = 1000;
 	// bUseWorldAxes gets updated automatically as part of UViewAdjustedStaticMeshGizmoComponent::UpdateWorldLocalState
 	TransformAdjuster->SetSettings(Settings);
-	if (ComponentToKeepDistanceConstantTo)
+	if (GizmoRootComponent)
 	{
 		TransformAdjuster->SetGizmoOriginTransform(
-			ComponentToKeepDistanceConstantTo->GetComponentTransform().GetRelativeTransform(ComponentIn->GetComponentToWorld()));
+			GizmoRootComponent->GetComponentTransform().GetRelativeTransform(ComponentIn->GetComponentToWorld()));
 	}
 	ComponentIn->SetTransformAdjuster(TransformAdjuster);
 
diff --git a/Engine/Source/Runtime/InteractiveToolsFramework/Public/BaseGizmos/ViewBasedTransformAdjusters.h b/Engine/Source/Runtime/InteractiveToolsFramework/Public/BaseGizmos/ViewBasedTransformAdjusters.h
--- a/Engine/Source/Runtime/InteractiveToolsFramework/Public/BaseGizmos/ViewBasedTransformAdjusters.h
+++ b/Engine/Source/Runtime/InteractiveToolsFramework/Public/BaseGizmos/ViewBasedTransformAdjusters.h
@@ -134,15 +134,17 @@ namespace UE::GizmoRenderingUtil
 		 * Static helper method to create and add this adjuster to a gizmo component for the common case of keeping a 
 		 *  constant size gizmo.
 		 * 
+		 * @param GizmoRootComponent The distance relative to this component is kept constant, and this is the component whose rotation
+		 *  is considered to be 0 when using global mode.
 		 * @param bMirrorBasedOnOctant If true, mirrors the component around the ComponentToKeepDistanceConstantTo
 		 *  depending on which octant of that component the view is located in.
 		 */
 		INTERACTIVETOOLSFRAMEWORK_API static TSharedPtr<FSubGizmoTransformAdjuster> AddTransformAdjuster(
-			UViewAdjustedStaticMeshGizmoComponent* ComponentIn, USceneComponent* ComponentToKeepDistanceConstantTo, 
+			UViewAdjustedStaticMeshGizmoComponent* ComponentIn, USceneComponent* GizmoRootComponent, 
 			bool bMirrorBasedOnOctant);
 
 		INTERACTIVETOOLSFRAMEWORK_API static TSharedPtr<FSubGizmoTransformAdjuster> AddTransformAdjuster(
-			UViewAdjustedStaticMeshGizmoComponent* ComponentIn, USceneComponent* ComponentToKeepDistanceConstantTo,
+			UViewAdjustedStaticMeshGizmoComponent* ComponentIn, USceneComponent* GizmoRootComponent,
 			TFunction<bool()> ShouldMirrorBasedOnOctant);
 
 
