InteractiveToolsFramework: Disallow FSubGizmoTransformAdjuster from adjusting gizmo origin in ortho mode when doing component line traces.

diff --git a/Engine/Source/Runtime/InteractiveToolsFramework/Private/BaseGizmos/ViewBasedTransformAdjusters.cpp b/Engine/Source/Runtime/InteractiveToolsFramework/Private/BaseGizmos/ViewBasedTransformAdjusters.cpp
--- a/Engine/Source/Runtime/InteractiveToolsFramework/Private/BaseGizmos/ViewBasedTransformAdjusters.cpp
+++ b/Engine/Source/Runtime/InteractiveToolsFramework/Private/BaseGizmos/ViewBasedTransformAdjusters.cpp
@@ -49,11 +49,14 @@ namespace ViewBasedTransformAdjustersLocals
 		const UE::GizmoRenderingUtil::ISceneViewInterface& View,
 		const FTransform& CurrentComponentToWorld,
 		const FTransform& GizmoOriginToComponent,
-		const UE::GizmoRenderingUtil::FSubGizmoTransformAdjuster::FSettings& Settings)
+		const UE::GizmoRenderingUtil::FSubGizmoTransformAdjuster::FSettings& Settings, 
+		bool bAllowGizmoOriginTranslation)
 	{
 		bool bIsOrtho = !View.IsPerspectiveProjection();
+		bool bApplyOrthoGizmoTranslation = bIsOrtho && bAllowGizmoOriginTranslation && Settings.DistanceInFrontOfCameraInOrtho >= 0;
+
 		if (!Settings.bKeepConstantViewSize && !Settings.bMirrorBasedOnOctant && !Settings.bUseWorldAxesForGizmo 
-			&& (Settings.DistanceInFrontOfCameraInOrtho <= 0 || !bIsOrtho))
+			&& !bApplyOrthoGizmoTranslation)
 		{
 			return CurrentComponentToWorld;
 		}
@@ -61,7 +64,7 @@ namespace ViewBasedTransformAdjustersLocals
 		FTransform GizmoToWorld = GizmoOriginToComponent * CurrentComponentToWorld;
 		FTransform ComponentToGizmo = CurrentComponentToWorld.GetRelativeTransform(GizmoToWorld);
 
-		if (bIsOrtho && Settings.DistanceInFrontOfCameraInOrtho >= 0)
+		if (bApplyOrthoGizmoTranslation)
 		{
 			UE::Geometry::FFrame3d CameraFrame(View.GetViewLocation(), View.GetViewDirection(), View.GetViewRight(), View.GetViewUp());
 			FVector CameraSpaceLocation = CameraFrame.ToFramePoint(GizmoToWorld.GetLocation());
@@ -184,15 +187,25 @@ TSharedPtr<UE::GizmoRenderingUtil::FSubGizmoTransformAdjuster> UE::GizmoRenderin
 FTransform UE::GizmoRenderingUtil::FSubGizmoTransformAdjuster::GetAdjustedComponentToWorld(
 	const ISceneViewInterface& View, const FTransform& CurrentComponentToWorld)
 {
+	// For rendering, we allow the gizmo origin to move towards the camera to avoid clipping in ortho. However in
+	//  line traces on the game thread, which are used in calculations in a gizmo, we typically need the hit location
+	//  to be relative to an unchanging gizmo origin, because that is the origin that will be used by other calculations.
+	// This is particularly important for axis rotation gizmos where the initial hit location (projected to the axis
+	//  passing through gizmo origin) determines the plane in which we rotate.
+	const bool bAllowGizmoOriginTranslation = false;
+
 	return ViewBasedTransformAdjustersLocals::GetSubGizmoAdjustedTransform(View, CurrentComponentToWorld,
-		GizmoOriginToComponent_GameThread, Settings);
+		GizmoOriginToComponent_GameThread, Settings, bAllowGizmoOriginTranslation);
 }
 
 FTransform UE::GizmoRenderingUtil::FSubGizmoTransformAdjuster::GetAdjustedComponentToWorld_RenderThread(
 	const ISceneViewInterface& View, const FTransform& CurrentComponentToWorld)
 {
+	// See comment in GetAdjustedComponentToWorld
+	const bool bAllowGizmoOriginTranslation = true;
+
 	return ViewBasedTransformAdjustersLocals::GetSubGizmoAdjustedTransform(View, CurrentComponentToWorld,
-		GizmoOriginToComponent_RenderThread, Settings);
+		GizmoOriginToComponent_RenderThread, Settings, bAllowGizmoOriginTranslation);
 }
 
 FBoxSphereBounds UE::GizmoRenderingUtil::FSubGizmoTransformAdjuster::GetViewIndependentBounds(
