InteractiveToolsFramework: Add ability to add custom delta constraint functions to CombinedTransformGizmo.

diff --git a/Engine/Source/Runtime/InteractiveToolsFramework/Private/BaseGizmos/CombinedTransformGizmo.cpp b/Engine/Source/Runtime/InteractiveToolsFramework/Private/BaseGizmos/CombinedTransformGizmo.cpp
--- a/Engine/Source/Runtime/InteractiveToolsFramework/Private/BaseGizmos/CombinedTransformGizmo.cpp
+++ b/Engine/Source/Runtime/InteractiveToolsFramework/Private/BaseGizmos/CombinedTransformGizmo.cpp
@@ -1040,6 +1040,36 @@ void UCombinedTransformGizmo::SetWorldAlignmentFunctions(
 	}
 }
 
+void UCombinedTransformGizmo::SetCustomTranslationDeltaFunctions(
+	TFunction<bool(double AxisDelta, double& SnappedDelta)> XAxis, 
+	TFunction<bool(double AxisDelta, double& SnappedDelta)> YAxis, 
+	TFunction<bool(double AxisDelta, double& SnappedDelta)> ZAxis)
+{
+	CustomTranslationDeltaConstraintFunctions[0] = XAxis;
+	CustomTranslationDeltaConstraintFunctions[1] = YAxis;
+	CustomTranslationDeltaConstraintFunctions[2] = ZAxis;
+}
+
+void UCombinedTransformGizmo::SetCustomRotationDeltaFunctions(
+	TFunction<bool(double AxisDelta, double& SnappedDelta)> XAxis,
+	TFunction<bool(double AxisDelta, double& SnappedDelta)> YAxis, 
+	TFunction<bool(double AxisDelta, double& SnappedDelta)> ZAxis)
+{
+	CustomRotationDeltaConstraintFunctions[0] = XAxis;
+	CustomRotationDeltaConstraintFunctions[1] = YAxis;
+	CustomRotationDeltaConstraintFunctions[2] = ZAxis;
+}
+
+void UCombinedTransformGizmo::SetCustomScaleDeltaFunctions(
+	TFunction<bool(double AxisDelta, double& SnappedDelta)> XAxis, 
+	TFunction<bool(double AxisDelta, double& SnappedDelta)> YAxis, 
+	TFunction<bool(double AxisDelta, double& SnappedDelta)> ZAxis)
+{
+	CustomScaleDeltaConstraintFunctions[0] = XAxis;
+	CustomScaleDeltaConstraintFunctions[1] = YAxis;
+	CustomScaleDeltaConstraintFunctions[2] = ZAxis;
+}
+
 void UCombinedTransformGizmo::SetDisallowNegativeScaling(bool bDisallow)
 {
 	if (bDisallowNegativeScaling != bDisallow)
@@ -1590,7 +1620,8 @@ UInteractiveGizmo* UCombinedTransformGizmo::AddAxisRotationGizmo(
 
 	if (UGizmoAxisRotationParameterSource* AngleSource = Cast<UGizmoAxisRotationParameterSource>(Gizmo->AngleSource.GetObject()))
 	{
-		AngleSource->AngleDeltaConstraintFunction = [this](double AngleDelta, double& SnappedDelta) { return RotationAxisAngleSnapFunction(AngleDelta, SnappedDelta, 0); };
+		int AxisIndex = Params.GetClampedAxisIndex();
+		AngleSource->AngleDeltaConstraintFunction = [this, AxisIndex](double AngleDelta, double& SnappedDelta) { return RotationAxisAngleSnapFunction(AngleDelta, SnappedDelta, AxisIndex); };
 	}
 	else
 	{
@@ -1616,7 +1647,8 @@ UInteractiveGizmo* UCombinedTransformGizmo::AddAxisScaleGizmo(
 
 	if (UGizmoAxisScaleParameterSource* ParameterSource = Cast<UGizmoAxisScaleParameterSource>(Gizmo->ParameterSource.GetObject()))
 	{
-		ParameterSource->ScaleAxisDeltaConstraintFunction = [this](const double ScaleAxisDelta, double& SnappedScaleAxisDelta) { return ScaleAxisDeltaSnapFunction(ScaleAxisDelta, SnappedScaleAxisDelta); };
+		int AxisIndex = Params.GetClampedAxisIndex();
+		ParameterSource->ScaleAxisDeltaConstraintFunction = [this, AxisIndex](const double ScaleAxisDelta, double& SnappedScaleAxisDelta) { return ScaleAxisDeltaSnapFunction(ScaleAxisDelta, SnappedScaleAxisDelta, AxisIndex); };
 	}
 	else
 	{
@@ -1642,8 +1674,11 @@ UInteractiveGizmo* UCombinedTransformGizmo::AddPlaneScaleGizmo(
 
 	if (UGizmoPlaneScaleParameterSource* ParameterSource = Cast<UGizmoPlaneScaleParameterSource>(Gizmo->ParameterSource.GetObject()))
 	{
-		ParameterSource->ScaleAxisXDeltaConstraintFunction = [this](double ScaleAxisDelta, double& SnappedScaleAxisDelta) { return ScaleAxisDeltaSnapFunction(ScaleAxisDelta, SnappedScaleAxisDelta); };
-		ParameterSource->ScaleAxisYDeltaConstraintFunction = [this](double ScaleAxisDelta, double& SnappedScaleAxisDelta) { return ScaleAxisDeltaSnapFunction(ScaleAxisDelta, SnappedScaleAxisDelta); };
+		int AxisIndex = Params.GetClampedAxisIndex();
+		int XAxes[3] = { 1, 2, 0 };
+		int YAxes[3] = { 2, 0, 1 };
+		ParameterSource->ScaleAxisXDeltaConstraintFunction = [this, XAxisIndex = XAxes[AxisIndex]](double ScaleAxisDelta, double& SnappedScaleAxisDelta) { return ScaleAxisDeltaSnapFunction(ScaleAxisDelta, SnappedScaleAxisDelta, XAxisIndex); };
+		ParameterSource->ScaleAxisYDeltaConstraintFunction = [this, YAxisIndex = YAxes[AxisIndex]](double ScaleAxisDelta, double& SnappedScaleAxisDelta) { return ScaleAxisDeltaSnapFunction(ScaleAxisDelta, SnappedScaleAxisDelta, YAxisIndex); };
 	}
 	else
 	{
@@ -2018,6 +2053,11 @@ bool UCombinedTransformGizmo::PositionSnapFunction(const FVector& WorldPosition,
 
 bool UCombinedTransformGizmo::PositionAxisDeltaSnapFunction(double AxisDelta, double& SnappedDeltaOut, int AxisIndex) const
 {
+	if (CustomTranslationDeltaConstraintFunctions[AxisIndex])
+	{
+		return CustomTranslationDeltaConstraintFunctions[AxisIndex](AxisDelta, SnappedDeltaOut);
+	}
+
 	if (!bSnapToWorldGrid) return false;
 
 	EToolContextCoordinateSystem CoordSystem = GetGizmoManager()->GetContextQueriesAPI()->GetCurrentCoordinateSystem();
@@ -2104,6 +2144,11 @@ FQuat UCombinedTransformGizmo::RotationSnapFunction(const FQuat& DeltaRotation)
 
 bool UCombinedTransformGizmo::RotationAxisAngleSnapFunction(double AxisAngleDelta, double& SnappedAxisAngleDeltaOut, int AxisIndex) const
 {
+	if (CustomRotationDeltaConstraintFunctions[AxisIndex])
+	{
+		return CustomRotationDeltaConstraintFunctions[AxisIndex](AxisAngleDelta, SnappedAxisAngleDeltaOut);
+	}
+
 	if (!bSnapToWorldRotGrid) return false;
 
 	FToolContextSnappingConfiguration SnappingConfig = GetGizmoManager()->GetContextQueriesAPI()->GetCurrentSnappingSettings();
@@ -2123,6 +2168,16 @@ bool UCombinedTransformGizmo::RotationAxisAngleSnapFunction(double AxisAngleDelt
 	return false;
 }
 
+bool UCombinedTransformGizmo::ScaleAxisDeltaSnapFunction(const double ScaleAxisDelta, double& SnappedAxisScaleDeltaOut, int AxisIndex) const
+{
+	if (CustomScaleDeltaConstraintFunctions[AxisIndex])
+	{
+		return CustomScaleDeltaConstraintFunctions[AxisIndex](ScaleAxisDelta, SnappedAxisScaleDeltaOut);
+	}
+
+	return ScaleAxisDeltaSnapFunction(ScaleAxisDelta, SnappedAxisScaleDeltaOut);
+}
+
 bool UCombinedTransformGizmo::ScaleAxisDeltaSnapFunction(const double ScaleAxisDelta, double & SnappedAxisScaleDeltaOut) const
 {
 	if (!bSnapToScaleGrid) return false;
diff --git a/Engine/Source/Runtime/InteractiveToolsFramework/Public/BaseGizmos/CombinedTransformGizmo.h b/Engine/Source/Runtime/InteractiveToolsFramework/Public/BaseGizmos/CombinedTransformGizmo.h
--- a/Engine/Source/Runtime/InteractiveToolsFramework/Public/BaseGizmos/CombinedTransformGizmo.h
+++ b/Engine/Source/Runtime/InteractiveToolsFramework/Public/BaseGizmos/CombinedTransformGizmo.h
@@ -373,6 +373,25 @@ public:
 		TUniqueFunction<bool(const FRay&, FVector&)>&& DestinationAlignmentRayCaster
 		);
 
+	/**
+	 * These allow for the deltas of gizmo manipulations to be constrained or clamped in custom ways, for instance
+	 *  to slow or stop the gizmo as the drag gets longer. The deltas constrained here are relative to drag start,
+	 *  and note that a custom constraint stops default world grid delta snapping from being applied on that axis.
+	 * Providing nullptr to any of these removes the custom constraint.
+	 */
+	INTERACTIVETOOLSFRAMEWORK_API void SetCustomTranslationDeltaFunctions(
+		TFunction<bool(double AxisDelta, double& SnappedDelta)> XAxis,
+		TFunction<bool(double AxisDelta, double& SnappedDelta)> YAxis,
+		TFunction<bool(double AxisDelta, double& SnappedDelta)> ZAxis);
+	INTERACTIVETOOLSFRAMEWORK_API void SetCustomRotationDeltaFunctions(
+		TFunction<bool(double AxisDelta, double& SnappedDelta)> XAxis,
+		TFunction<bool(double AxisDelta, double& SnappedDelta)> YAxis,
+		TFunction<bool(double AxisDelta, double& SnappedDelta)> ZAxis);
+	INTERACTIVETOOLSFRAMEWORK_API void SetCustomScaleDeltaFunctions(
+		TFunction<bool(double AxisDelta, double& SnappedDelta)> XAxis,
+		TFunction<bool(double AxisDelta, double& SnappedDelta)> YAxis,
+		TFunction<bool(double AxisDelta, double& SnappedDelta)> ZAxis);
+
 	/**
 	 * By default, non-uniform scaling handles appear (assuming they exist in the gizmo to begin with), 
 	 * when CurrentCoordinateSystem == EToolContextCoordinateSystem::Local, since components can only be
@@ -637,6 +656,10 @@ private:
 	TWeakObjectPtr<UInteractiveGizmo> AxisScaleYGizmo;
 	TWeakObjectPtr<UInteractiveGizmo> AxisScaleZGizmo;
 
+	TFunction<bool(double AxisDelta, double& SnappedDelta)> CustomTranslationDeltaConstraintFunctions[3];
+	TFunction<bool(double AxisDelta, double& SnappedDelta)> CustomRotationDeltaConstraintFunctions[3];
+	TFunction<bool(double AxisDelta, double& SnappedDelta)> CustomScaleDeltaConstraintFunctions[3];
+
 protected:
 	/** GizmoActors will be spawned in this World */
 	UWorld* World;
@@ -755,12 +778,15 @@ protected:
 	INTERACTIVETOOLSFRAMEWORK_API FQuat RotationSnapFunction(const FQuat& DeltaRotation) const;
 	INTERACTIVETOOLSFRAMEWORK_API bool RotationAxisAngleSnapFunction(double AxisAngleDelta, double& SnappedAxisAngleDeltaOut, int AxisIndex) const;
 	
+private:
+	// Used for scale delta snapping. Calls out to the non-axis endpoint by default
+	bool ScaleAxisDeltaSnapFunction(double ScaleAxisDelta, double& SnappedAxisScaleDeltaOut, int AxisIndex) const;
+protected:
 	// currently not implemented because WorldGridSnapping currently has no affect on scale/scale-snapping
 	INTERACTIVETOOLSFRAMEWORK_API bool ScaleSnapFunction(double DeltaScale) const;
-	// used to execute Plane, Axis, & Uniform scale snapping
+	// Used for uniform scale delta snapping
 	INTERACTIVETOOLSFRAMEWORK_API bool ScaleAxisDeltaSnapFunction(double ScaleAxisDelta, double & SnappedAxisScaleDeltaOut) const;
 
-
 	UE_DEPRECATED(5.5, "Use FTransformSubGizmoCommonParams overload instead.")
 	/** @return a new instance of the standard axis-translation Gizmo */
 	INTERACTIVETOOLSFRAMEWORK_API virtual UInteractiveGizmo* AddAxisTranslationGizmo(
