ScriptableTools: Added ensure check for re-entrant FTransformProxySource::EndChange calls.
Deregister callbacks on DestroyTRSGizmo.

diff --git a/Engine/Plugins/Runtime/ScriptableToolsFramework/Source/ScriptableToolsFramework/Private/ScriptableInteractiveTool.cpp b/Engine/Plugins/Runtime/ScriptableToolsFramework/Source/ScriptableToolsFramework/Private/ScriptableInteractiveTool.cpp
--- a/Engine/Plugins/Runtime/ScriptableToolsFramework/Source/ScriptableToolsFramework/Private/ScriptableInteractiveTool.cpp
+++ b/Engine/Plugins/Runtime/ScriptableToolsFramework/Source/ScriptableToolsFramework/Private/ScriptableInteractiveTool.cpp
@@ -1078,6 +1078,14 @@ void UScriptableInteractiveTool::DestroyTRSGizmo(
 		return;
 	}
 
+	if ((*Found)->ActiveTarget)
+	{
+		(*Found)->ActiveTarget->OnTransformChanged.RemoveAll(this);
+		(*Found)->ActiveTarget->OnBeginTransformEdit.RemoveAll(this);
+		(*Found)->ActiveTarget->OnEndTransformEdit.RemoveAll(this);
+		(*Found)->ActiveTarget->OnTransformChangedUndoRedo.RemoveAll(this);
+	}
+
 	GetToolManager()->GetPairedGizmoManager()->DestroyGizmo(*Found);
 	Gizmos.Remove(Identifier);
 	Outcome = EToolsFrameworkOutcomePins::Success;
diff --git a/Engine/Source/Runtime/InteractiveToolsFramework/Private/BaseGizmos/TransformProxy.cpp b/Engine/Source/Runtime/InteractiveToolsFramework/Private/BaseGizmos/TransformProxy.cpp
--- a/Engine/Source/Runtime/InteractiveToolsFramework/Private/BaseGizmos/TransformProxy.cpp
+++ b/Engine/Source/Runtime/InteractiveToolsFramework/Private/BaseGizmos/TransformProxy.cpp
@@ -8,6 +8,7 @@
 
 #define LOCTEXT_NAMESPACE "UTransformProxy"
 
+DEFINE_LOG_CATEGORY_STATIC(LogInteractiveToolsTransformProxy, Log, All);
 
 void UTransformProxy::AddComponent(USceneComponent* Component, bool bModifyComponentOnTransform)
 {
@@ -257,6 +258,13 @@ TUniquePtr<FToolCommandChange> FTransformProxyChangeSource::EndChange()
 			Proxy->EndTransformEditSequence();
 		}
 
+		if (!ActiveChange)
+		{
+			// One way this error could occur is if the Proxy call forced a re-entrant call to EndChange() by destroying
+			// the gizmo which in turn will ForceTerminateSource() and end an ActiveChange.
+			UE_LOG(LogInteractiveToolsTransformProxy, Error, TEXT("FTransformProxyChangeSource: Re-entrant EndChange invalidated ActiveChange. Did you delete a gizmo on EndPivotEditSequence/EndTransformEditSequence?"));
+			return TUniquePtr<FToolCommandChange>();
+		}
 		ActiveChange->To = Proxy->GetTransform();
 		return MoveTemp(ActiveChange);
 	}
