add default values for toolcontextinterfaces FViewCameraState

diff --git a/Engine/Source/Runtime/InteractiveToolsFramework/Public/ToolContextInterfaces.h b/Engine/Source/Runtime/InteractiveToolsFramework/Public/ToolContextInterfaces.h
--- a/Engine/Source/Runtime/InteractiveToolsFramework/Public/ToolContextInterfaces.h
+++ b/Engine/Source/Runtime/InteractiveToolsFramework/Public/ToolContextInterfaces.h
@@ -81,19 +81,19 @@ struct FToolBuilderState
 struct FViewCameraState
 {
 	/** Current camera/head position */
-	FVector Position;
+	FVector Position = FVector(EForceInit::ForceInit);
 	/** Current camera/head orientation */
-	FQuat Orientation;
+	FQuat Orientation = FQuat(EForceInit::ForceInit);
 	/** Current Horizontal Field of View Angle in Degrees. Only relevant if bIsOrthographic is false. */
-	float HorizontalFOVDegrees;
+	float HorizontalFOVDegrees = 90;
 	/** Current width of viewport in world space coordinates. Only valid if bIsOrthographic is true. */
-	float OrthoWorldCoordinateWidth;
+	float OrthoWorldCoordinateWidth = 100;
 	/** Current Aspect Ratio */
-	float AspectRatio;
+	float AspectRatio = 1;
 	/** Is current view an orthographic view */
-	bool bIsOrthographic;
+	bool bIsOrthographic = false;
 	/** Is current view a VR view */
-	bool bIsVR;
+	bool bIsVR = false;
 
 	/** @return "right"/horizontal direction in camera plane */
 	FVector Right() const { return Orientation.GetAxisY(); }
