InteractiveToolsFramework: Fix sporadic crash from a memory stomp in FSubGizmoTransformAdjuster.

diff --git a/Engine/Source/Runtime/InteractiveToolsFramework/Private/BaseGizmos/ViewBasedTransformAdjusters.cpp b/Engine/Source/Runtime/InteractiveToolsFramework/Private/BaseGizmos/ViewBasedTransformAdjusters.cpp
--- a/Engine/Source/Runtime/InteractiveToolsFramework/Private/BaseGizmos/ViewBasedTransformAdjusters.cpp
+++ b/Engine/Source/Runtime/InteractiveToolsFramework/Private/BaseGizmos/ViewBasedTransformAdjusters.cpp
@@ -149,9 +149,13 @@ void UE::GizmoRenderingUtil::FSubGizmoTransformAdjuster::SetGizmoOriginTransform
 
 	// Safely update the render thread side transform.
 	ENQUEUE_RENDER_COMMAND(FViewpointOctantMirrorTransformAdjusterUpdate)(
-		[GizmoOriginToComponentIn, this](FRHICommandListImmediate& RHICmdList)
+		[GizmoOriginToComponentIn, WeakThis = AsWeak()](FRHICommandListImmediate& RHICmdList)
 	{
-		GizmoOriginToComponent_RenderThread = GizmoOriginToComponentIn;
+		TSharedPtr<UE::GizmoRenderingUtil::FSubGizmoTransformAdjuster> ThisPinned = WeakThis.Pin();
+		if (ThisPinned.IsValid())
+		{
+			ThisPinned->GizmoOriginToComponent_RenderThread = GizmoOriginToComponentIn;
+		}
 	});
 }
 
diff --git a/Engine/Source/Runtime/InteractiveToolsFramework/Public/BaseGizmos/ViewBasedTransformAdjusters.h b/Engine/Source/Runtime/InteractiveToolsFramework/Public/BaseGizmos/ViewBasedTransformAdjusters.h
--- a/Engine/Source/Runtime/InteractiveToolsFramework/Public/BaseGizmos/ViewBasedTransformAdjusters.h
+++ b/Engine/Source/Runtime/InteractiveToolsFramework/Public/BaseGizmos/ViewBasedTransformAdjusters.h
@@ -22,6 +22,8 @@ namespace UE::GizmoRenderingUtil
 	class IViewBasedTransformAdjuster
 	{
 	public:
+		virtual ~IViewBasedTransformAdjuster() {}
+
 		/**
 		 * Given the component location and the view information, gives the desired transform of the component.
 		 */
@@ -86,7 +88,9 @@ namespace UE::GizmoRenderingUtil
 	 * An adjuster that can do various transformations common for sub gizmos, which are based off
 	 *  of the parent gizmo transform (in addition to the actual component transform).
 	 */
-	class FSubGizmoTransformAdjuster : public IViewBasedTransformAdjuster
+	class FSubGizmoTransformAdjuster : public IViewBasedTransformAdjuster,
+		// Used for getting a weak pointer to itself when enqueuing state change on the render thread (in SetGizmoOriginTransform)
+		public TSharedFromThis<FSubGizmoTransformAdjuster, ESPMode::ThreadSafe>
 	{
 	public:
 		struct FSettings
