[ITF] Dashed line rendering

diff --git a/Engine/Source/Runtime/InteractiveToolsFramework/Private/BaseGizmos/GizmoElementArrow.cpp b/Engine/Source/Runtime/InteractiveToolsFramework/Private/BaseGizmos/GizmoElementArrow.cpp
--- a/Engine/Source/Runtime/InteractiveToolsFramework/Private/BaseGizmos/GizmoElementArrow.cpp
+++ b/Engine/Source/Runtime/InteractiveToolsFramework/Private/BaseGizmos/GizmoElementArrow.cpp
@@ -71,7 +71,7 @@ void UGizmoElementArrow::Render(IToolsContextRenderAPI* RenderAPI, const FRender
 FInputRayHit UGizmoElementArrow::LineTrace(const UGizmoViewContext* ViewContext, const FLineTraceTraversalState& LineTraceState, const FVector& RayOrigin, const FVector& RayDirection)
 {
 	using namespace UE::Private;
-	
+
 	FLineTraceTraversalState CurrentLineTraceState(LineTraceState);
 	bool bHittableViewDependent = UpdateLineTraceState(ViewContext, Base, CurrentLineTraceState);
 
@@ -96,7 +96,7 @@ FInputRayHit UGizmoElementArrow::LineTrace(const UGizmoViewContext* ViewContext,
 			static constexpr float NoPixelHitThreshold = 0.f;
 			UpdatePixelThreshold(CylinderElement, NoPixelHitThreshold);
 			UpdatePixelThreshold(ConeElement, NoPixelHitThreshold);
-			UpdatePixelThreshold(BoxElement, NoPixelHitThreshold);		
+			UpdatePixelThreshold(BoxElement, NoPixelHitThreshold);
 		}
 	}
 	
@@ -293,6 +293,40 @@ EGizmoElementArrowHeadType UGizmoElementArrow::GetHeadType() const
 	return HeadType;
 }
 
+void UGizmoElementArrow::SetIsDashed(bool bInDashing)
+{
+	if (ensure(CylinderElement))
+	{
+		CylinderElement->SetIsDashed(bInDashing);
+	}
+}
+
+bool UGizmoElementArrow::IsDashed() const
+{
+	if (ensure(CylinderElement))
+	{
+		return CylinderElement->GetIsDashed();
+	}
+
+	return false;
+}
+
+void UGizmoElementArrow::SetDashParameters(const float InDashLength, const TOptional<float>& InGapLength)
+{
+	if (ensure(CylinderElement))
+	{
+		CylinderElement->SetDashParameters(InDashLength, InGapLength);
+	}
+}
+
+void UGizmoElementArrow::GetDashParameters(float& OutDashLength, float& OutGapLength) const
+{
+	if (ensure(CylinderElement))
+	{
+		CylinderElement->GetDashParameters(OutDashLength, OutGapLength);
+	}
+}
+
 void UGizmoElementArrow::UpdateArrowBody()
 {
 	CylinderElement->SetBase(FVector::ZeroVector);
@@ -311,7 +345,7 @@ void UGizmoElementArrow::UpdateArrowHead()
 	{
 		check(ConeElement);
 		// head length is multiplied by 0.9f to prevent gap between body cylinder and head cone
-		ConeElement->SetOrigin(Direction * (BodyLength + HeadLength * 0.9f)); 
+		ConeElement->SetOrigin(Direction * (BodyLength + HeadLength * 0.9f));
 		ConeElement->SetDirection(-Direction);
 		ConeElement->SetHeight(HeadLength);
 		ConeElement->SetRadius(HeadRadius);
diff --git a/Engine/Source/Runtime/InteractiveToolsFramework/Private/BaseGizmos/GizmoElementCylinder.cpp b/Engine/Source/Runtime/InteractiveToolsFramework/Private/BaseGizmos/GizmoElementCylinder.cpp
--- a/Engine/Source/Runtime/InteractiveToolsFramework/Private/BaseGizmos/GizmoElementCylinder.cpp
+++ b/Engine/Source/Runtime/InteractiveToolsFramework/Private/BaseGizmos/GizmoElementCylinder.cpp
@@ -11,6 +11,29 @@
 
 #include UE_INLINE_GENERATED_CPP_BY_NAME(GizmoElementCylinder)
 
+namespace GizmoElementCylinderLocals
+{
+	/**
+	 * Normalizes the given dash parameters to fit the input range.
+	 * @param	OutScaleFactor the scale factor to apply to the dash length and gap length.
+	 * @return	the number of dashes that fit within the range.
+	 */
+	template < typename RealType = float>
+	int32 GetDashNumAndScale(const RealType& InLineLength, const RealType& InDashLength, const RealType& InDashGapLength, RealType& OutScaleFactor)
+	{
+		const RealType DashLengthAndGap = InDashLength + InDashGapLength;
+		OutScaleFactor = static_cast<RealType>(1);
+
+		// Number based on sequence of dash, gap pairs, and last dash. This should always be an even number for fit alignment
+		const int32 NumDashes = FMath::Max(2, FMath::CeilToInt32((InLineLength / DashLengthAndGap)));
+
+		// Get a scale factor based on how many dashes with gaps (and an end dash) can actually fit on this line
+		OutScaleFactor = InLineLength / (((NumDashes - 1) * DashLengthAndGap) + InDashLength);
+
+		return NumDashes;
+	}
+}
+
 void UGizmoElementCylinder::Render(IToolsContextRenderAPI* RenderAPI, const FRenderTraversalState& RenderState)
 {
 	FRenderTraversalState CurrentRenderState(RenderState);
@@ -26,7 +49,53 @@ void UGizmoElementCylinder::Render(IToolsContextRenderAPI* RenderAPI, const FRen
 
 			FTransform RenderLocalToWorldTransform = FTransform(Rotation, OriginOffset) * CurrentRenderState.LocalToWorldTransform;
 			FPrimitiveDrawInterface* PDI = RenderAPI->GetPrimitiveDrawInterface();
-			DrawCylinder(PDI, RenderLocalToWorldTransform.ToMatrixWithScale(), FVector::ZeroVector, FVector(1, 0, 0), FVector(0, 1, 0), FVector(0, 0, 1), Radius, HalfHeight, NumSides, UseMaterial->GetRenderProxy(), SDPG_Foreground);
+
+			if (bIsDashed)
+			{
+				// Get a scale factor based on how many dashes with gaps (and an end dash) can actually fit on this line
+				float DashLengthScale = 0.0f;
+
+				// Number based on sequence of dash, gap pairs, and last dash
+				const int32 NumDashes = GizmoElementCylinderLocals::GetDashNumAndScale(FMath::Max(0.0f, Height - DashLength), DashLength, DashGapLength, DashLengthScale);
+
+				const float AdjustedDashLength = DashLength * DashLengthScale;
+				const float AdjustDashGapLength = DashGapLength * DashLengthScale;
+
+				const float DashLengthAndGaps = (AdjustedDashLength + AdjustDashGapLength);
+
+				const float HalfDashLength = AdjustedDashLength * 0.5f;
+
+				for (int32 DashIndex = 0; DashIndex <= NumDashes; ++DashIndex)
+				{
+					DrawCylinder(
+						PDI,
+						RenderLocalToWorldTransform.ToMatrixWithScale(),
+						(-FVector::UpVector * HalfHeight) + FVector::UpVector * (HalfDashLength + (DashIndex * DashLengthAndGaps)),
+						FVector(1, 0, 0),
+						FVector(0, 1, 0),
+						FVector(0, 0, 1),
+						Radius,
+						HalfDashLength,
+						NumSides,
+						UseMaterial->GetRenderProxy(),
+						SDPG_Foreground);
+				}
+			}
+			else
+			{
+				DrawCylinder(
+					PDI,
+					RenderLocalToWorldTransform.ToMatrixWithScale(),
+					FVector::ZeroVector,
+					FVector(1, 0, 0),
+					FVector(0, 1, 0),
+					FVector(0, 0, 1),
+					Radius,
+					HalfHeight,
+					NumSides,
+					UseMaterial->GetRenderProxy(),
+					SDPG_Foreground);
+			}
 		}
 	}
 }
@@ -127,3 +196,25 @@ int32 UGizmoElementCylinder::GetNumSides() const
 	return NumSides;
 }
 
+void UGizmoElementCylinder::SetIsDashed(bool bInDashing)
+{
+	bIsDashed = bInDashing;
+}
+
+bool UGizmoElementCylinder::GetIsDashed() const
+{
+	return bIsDashed;
+}
+
+void UGizmoElementCylinder::SetDashParameters(const float InDashLength, const TOptional<float>& InGapLength)
+{
+	DashLength = InDashLength;
+	DashGapLength = InGapLength.Get(InDashLength * 0.5f);
+}
+
+void UGizmoElementCylinder::GetDashParameters(float& OutDashLength, float& OutGapLength) const
+{
+	OutDashLength = DashLength;
+	OutGapLength = DashGapLength;
+}
+
diff --git a/Engine/Source/Runtime/InteractiveToolsFramework/Public/BaseGizmos/GizmoElementArrow.h b/Engine/Source/Runtime/InteractiveToolsFramework/Public/BaseGizmos/GizmoElementArrow.h
--- a/Engine/Source/Runtime/InteractiveToolsFramework/Public/BaseGizmos/GizmoElementArrow.h
+++ b/Engine/Source/Runtime/InteractiveToolsFramework/Public/BaseGizmos/GizmoElementArrow.h
@@ -69,7 +69,7 @@ public:
 	// Number of sides for cylinder and cone, if relevant.
 	INTERACTIVETOOLSFRAMEWORK_API virtual void SetNumSides(int32 InNumSides);
 	INTERACTIVETOOLSFRAMEWORK_API virtual int32 GetNumSides() const;
-	
+
 	// Whether to render and end cap on the base of the arrow head. Applicable when the ArrowHeadType is Cone.
 	INTERACTIVETOOLSFRAMEWORK_API virtual void SetEndCaps(bool InEndCaps);
 	INTERACTIVETOOLSFRAMEWORK_API virtual bool GetEndCaps() const;
@@ -78,12 +78,20 @@ public:
 	INTERACTIVETOOLSFRAMEWORK_API virtual void SetHeadType(EGizmoElementArrowHeadType InHeadType);
 	INTERACTIVETOOLSFRAMEWORK_API virtual EGizmoElementArrowHeadType GetHeadType() const;
 
+	// Whether to render the body element with dashing.
+	INTERACTIVETOOLSFRAMEWORK_API virtual void SetIsDashed(bool bInDashing);
+	INTERACTIVETOOLSFRAMEWORK_API virtual bool IsDashed() const;
+
+	// The dash parameters for the body element, if enabled. If GapLength is not specified, it defaults to half DashLength.
+	INTERACTIVETOOLSFRAMEWORK_API virtual void SetDashParameters(const float InDashLength = 10.0f, const TOptional<float>& InGapLength = TOptional<float>());
+	INTERACTIVETOOLSFRAMEWORK_API virtual void GetDashParameters(float& OutDashLength, float& OutGapLength) const;
+
 	// Pixel hit distance threshold, element will be scaled enough to add this threshold when line-tracing.
 	INTERACTIVETOOLSFRAMEWORK_API virtual void SetPixelHitDistanceThreshold(float InPixelHitDistanceThreshold) override;
 
 	// Hit Mask Gizmo element used for adaptive pixel hit threshold
 	INTERACTIVETOOLSFRAMEWORK_API void SetHitMask(const TWeakObjectPtr<UGizmoElementBase>& InHitMask);
-	
+
 protected:
 
 	// Update arrow cylinder based on parameters
@@ -142,7 +150,7 @@ protected:
 	UPROPERTY()
 	int32 NumSides = 32;
 
-	// Whether to render and end cap on the base of the arrow head.
+	// Whether to render an end cap on the base of the arrow head.
 	UPROPERTY()
 	bool bEndCaps = false;
 
diff --git a/Engine/Source/Runtime/InteractiveToolsFramework/Public/BaseGizmos/GizmoElementCylinder.h b/Engine/Source/Runtime/InteractiveToolsFramework/Public/BaseGizmos/GizmoElementCylinder.h
--- a/Engine/Source/Runtime/InteractiveToolsFramework/Public/BaseGizmos/GizmoElementCylinder.h
+++ b/Engine/Source/Runtime/InteractiveToolsFramework/Public/BaseGizmos/GizmoElementCylinder.h
@@ -45,6 +45,14 @@ public:
 	INTERACTIVETOOLSFRAMEWORK_API virtual void SetNumSides(int32 InNumSides);
 	INTERACTIVETOOLSFRAMEWORK_API virtual int32 GetNumSides() const;
 
+	// Whether to render with dashing.
+	INTERACTIVETOOLSFRAMEWORK_API virtual void SetIsDashed(bool bInDashing);
+	INTERACTIVETOOLSFRAMEWORK_API virtual bool GetIsDashed() const;
+
+	// The dash parameter, if enabled. If GapLength is not specified, it defaults to half DashLength.
+	INTERACTIVETOOLSFRAMEWORK_API virtual void SetDashParameters(const float InDashLength = 10.0f, const TOptional<float>& InGapLength = TOptional<float>());
+	INTERACTIVETOOLSFRAMEWORK_API virtual void GetDashParameters(float& OutDashLength, float& OutGapLength) const;
+
 protected:
 
 	// Location of center of cylinder's base circle.
@@ -66,5 +74,16 @@ protected:
 	// Number of sides for tessellating cylinder.
 	UPROPERTY()
 	int32 NumSides = 32;
-};
 
+	// Whether to render with dashing.
+	UPROPERTY()
+	bool bIsDashed = false;
+
+	// The Dash Length, if Dashing is enabled.
+	UPROPERTY()
+	float DashLength = 10.0f;
+
+	// The Dash Gap Length, if Dashing is enabled.
+	UPROPERTY()
+	float DashGapLength = 5.0f;
+};
