[ITF] Added GizmoElementGroupBase, hiding Add/Remove and exposing them in the derived GizmoElementGroup

diff --git a/Engine/Source/Runtime/InteractiveToolsFramework/Private/BaseGizmos/GizmoElementGroup.cpp b/Engine/Source/Runtime/InteractiveToolsFramework/Private/BaseGizmos/GizmoElementGroup.cpp
--- a/Engine/Source/Runtime/InteractiveToolsFramework/Private/BaseGizmos/GizmoElementGroup.cpp
+++ b/Engine/Source/Runtime/InteractiveToolsFramework/Private/BaseGizmos/GizmoElementGroup.cpp
@@ -5,7 +5,7 @@
 
 #include UE_INLINE_GENERATED_CPP_BY_NAME(GizmoElementGroup)
 
-void UGizmoElementGroup::ApplyUniformConstantScaleToTransform(double PixelToWorldScale, FTransform& InOutLocalToWorldTransform) const
+void UGizmoElementGroupBase::ApplyUniformConstantScaleToTransform(double PixelToWorldScale, FTransform& InOutLocalToWorldTransform) const
 {
 	double Scale = InOutLocalToWorldTransform.GetScale3D().X;
 	if (bConstantScale)
@@ -15,7 +15,7 @@ void UGizmoElementGroup::ApplyUniformConstantScaleToTransform(double PixelToWorl
 	InOutLocalToWorldTransform.SetScale3D(FVector(Scale, Scale, Scale));
 }
 
-void UGizmoElementGroup::Render(IToolsContextRenderAPI* RenderAPI, const FRenderTraversalState& RenderState)
+void UGizmoElementGroupBase::Render(IToolsContextRenderAPI* RenderAPI, const FRenderTraversalState& RenderState)
 {
 	FRenderTraversalState CurrentRenderState(RenderState);
 	bool bVisibleViewDependent = UpdateRenderState(RenderAPI, FVector::ZeroVector, CurrentRenderState);
@@ -36,7 +36,7 @@ void UGizmoElementGroup::Render(IToolsContextRenderAPI* RenderAPI, const FRender
 	}
 }
 
-FInputRayHit UGizmoElementGroup::LineTrace(const UGizmoViewContext* ViewContext, const FLineTraceTraversalState& LineTraceState, const FVector& RayOrigin, const FVector& RayDirection)
+FInputRayHit UGizmoElementGroupBase::LineTrace(const UGizmoViewContext* ViewContext, const FLineTraceTraversalState& LineTraceState, const FVector& RayOrigin, const FVector& RayDirection)
 {
 	FInputRayHit Hit;
 
@@ -67,7 +67,7 @@ FInputRayHit UGizmoElementGroup::LineTrace(const UGizmoViewContext* ViewContext,
 	return Hit;
 }
 
-void UGizmoElementGroup::Add(UGizmoElementBase* InElement)
+void UGizmoElementGroupBase::Add(UGizmoElementBase* InElement)
 {
 	if (!Elements.Contains(InElement))
 	{
@@ -75,7 +75,7 @@ void UGizmoElementGroup::Add(UGizmoElementBase* InElement)
 	}
 }
 
-void UGizmoElementGroup::Remove(UGizmoElementBase* InElement)
+void UGizmoElementGroupBase::Remove(UGizmoElementBase* InElement)
 {
 	if (InElement)
 	{
@@ -87,7 +87,7 @@ void UGizmoElementGroup::Remove(UGizmoElementBase* InElement)
 	}
 }
 
-bool UGizmoElementGroup::UpdatePartVisibleState(bool bVisible, uint32 InPartIdentifier)
+bool UGizmoElementGroupBase::UpdatePartVisibleState(bool bVisible, uint32 InPartIdentifier)
 {
 	if (Super::UpdatePartVisibleState(bVisible, InPartIdentifier))
 	{
@@ -105,7 +105,7 @@ bool UGizmoElementGroup::UpdatePartVisibleState(bool bVisible, uint32 InPartIden
 	return false;
 }
 
-TOptional<bool> UGizmoElementGroup::GetPartVisibleState(uint32 InPartIdentifier) const
+TOptional<bool> UGizmoElementGroupBase::GetPartVisibleState(uint32 InPartIdentifier) const
 {
 	TOptional<bool> Result = Super::GetPartVisibleState(InPartIdentifier);
 	if (Result.IsSet())
@@ -129,7 +129,7 @@ TOptional<bool> UGizmoElementGroup::GetPartVisibleState(uint32 InPartIdentifier)
 }
 
 
-bool UGizmoElementGroup::UpdatePartHittableState(bool bHittable, uint32 InPartIdentifier)
+bool UGizmoElementGroupBase::UpdatePartHittableState(bool bHittable, uint32 InPartIdentifier)
 {
 	if (Super::UpdatePartHittableState(bHittable, InPartIdentifier))
 	{
@@ -147,7 +147,7 @@ bool UGizmoElementGroup::UpdatePartHittableState(bool bHittable, uint32 InPartId
 	return false;
 }
 
-TOptional<bool> UGizmoElementGroup::GetPartHittableState(uint32 InPartIdentifier) const
+TOptional<bool> UGizmoElementGroupBase::GetPartHittableState(uint32 InPartIdentifier) const
 {
 	TOptional<bool> Result = Super::GetPartHittableState(InPartIdentifier);
 	if (Result.IsSet())
@@ -170,7 +170,7 @@ TOptional<bool> UGizmoElementGroup::GetPartHittableState(uint32 InPartIdentifier
 	return TOptional<bool>();
 }
 
-bool UGizmoElementGroup::UpdatePartInteractionState(EGizmoElementInteractionState InInteractionState, uint32 InPartIdentifier)
+bool UGizmoElementGroupBase::UpdatePartInteractionState(EGizmoElementInteractionState InInteractionState, uint32 InPartIdentifier)
 {
 	if (Super::UpdatePartInteractionState(InInteractionState, InPartIdentifier))
 	{
@@ -189,7 +189,7 @@ bool UGizmoElementGroup::UpdatePartInteractionState(EGizmoElementInteractionStat
 }
 
 
-TOptional<EGizmoElementInteractionState> UGizmoElementGroup::GetPartInteractionState(uint32 InPartIdentifier) const
+TOptional<EGizmoElementInteractionState> UGizmoElementGroupBase::GetPartInteractionState(uint32 InPartIdentifier) const
 {
 	TOptional<EGizmoElementInteractionState> Result = Super::GetPartInteractionState(InPartIdentifier);
 	if (Result.IsSet())
@@ -212,22 +212,22 @@ TOptional<EGizmoElementInteractionState> UGizmoElementGroup::GetPartInteractionS
 	return TOptional<EGizmoElementInteractionState>();
 }
 
-void UGizmoElementGroup::SetConstantScale(bool bInConstantScale)
+void UGizmoElementGroupBase::SetConstantScale(bool bInConstantScale)
 {
 	bConstantScale = bInConstantScale;
 }
 
-bool UGizmoElementGroup::GetConstantScale() const
+bool UGizmoElementGroupBase::GetConstantScale() const
 {
 	return bConstantScale;
 }
 
-void UGizmoElementGroup::SetHitOwner(bool bInHitOwner)
+void UGizmoElementGroupBase::SetHitOwner(bool bInHitOwner)
 {
 	bHitOwner = bInHitOwner;
 }
 
-bool UGizmoElementGroup::GetHitOwner() const
+bool UGizmoElementGroupBase::GetHitOwner() const
 {
 	return bHitOwner;
 }
diff --git a/Engine/Source/Runtime/InteractiveToolsFramework/Public/BaseGizmos/GizmoElementGroup.h b/Engine/Source/Runtime/InteractiveToolsFramework/Public/BaseGizmos/GizmoElementGroup.h
--- a/Engine/Source/Runtime/InteractiveToolsFramework/Public/BaseGizmos/GizmoElementGroup.h
+++ b/Engine/Source/Runtime/InteractiveToolsFramework/Public/BaseGizmos/GizmoElementGroup.h
@@ -12,8 +12,8 @@
  * Simple group object intended to be used as part of 3D Gizmos.
  * Contains multiple gizmo objects.
  */
-UCLASS(Transient, MinimalAPI)
-class UGizmoElementGroup : public UGizmoElementLineBase
+UCLASS(Transient, Abstract, MinimalAPI)
+class UGizmoElementGroupBase : public UGizmoElementLineBase
 {
 	GENERATED_BODY()
 
@@ -24,12 +24,6 @@ public:
 	INTERACTIVETOOLSFRAMEWORK_API virtual FInputRayHit LineTrace(const UGizmoViewContext* ViewContext, const FLineTraceTraversalState& LineTraceState, const FVector& RayOrigin, const FVector& RayDirection) override;
 	//~ End UGizmoElementBase Interface.
 
-	// Add object to group.
-	INTERACTIVETOOLSFRAMEWORK_API virtual void Add(UGizmoElementBase* InElement);
-
-	// Remove object from group, if it exists
-	INTERACTIVETOOLSFRAMEWORK_API virtual void Remove(UGizmoElementBase* InElement);
-
 	// Update group and contained elements' visibility state for elements in specified gizmo parts, return true if part was found.
 	INTERACTIVETOOLSFRAMEWORK_API virtual bool UpdatePartVisibleState(bool bVisible, uint32 InPartIdentifier);
 
@@ -56,6 +50,14 @@ public:
 	INTERACTIVETOOLSFRAMEWORK_API virtual void SetHitOwner(bool bInHitOwner);
 	INTERACTIVETOOLSFRAMEWORK_API virtual bool GetHitOwner() const;
 
+protected:
+
+	// Add object to group.
+	INTERACTIVETOOLSFRAMEWORK_API virtual void Add(UGizmoElementBase* InElement);
+
+	// Remove object from group, if it exists
+	INTERACTIVETOOLSFRAMEWORK_API virtual void Remove(UGizmoElementBase* InElement);
+
 protected:
 
 	// When true, maintains view-dependent constant scale for this gizmo object hierarchy
@@ -74,3 +76,21 @@ protected:
 	// Updates input transform's scale component to have uniform scale and applies constant scale if bConstantScale is true
 	INTERACTIVETOOLSFRAMEWORK_API virtual void ApplyUniformConstantScaleToTransform(double PixelToWorldScale, FTransform& InOutLocalToWorldTransform) const;
 };
+
+/**
+ * Simple group object intended to be used as part of 3D Gizmos.
+ * Contains multiple gizmo objects.
+ */
+UCLASS(Transient, MinimalAPI)
+class UGizmoElementGroup : public UGizmoElementGroupBase
+{
+	GENERATED_BODY()
+
+public:
+
+	// Add object to group.
+	using UGizmoElementGroupBase::Add;
+
+	// Remove object from group, if it exists
+	using UGizmoElementGroupBase::Remove;
+};
