[ITF] GizmoElementSphere accounts for Scale

diff --git a/Engine/Source/Runtime/InteractiveToolsFramework/Private/BaseGizmos/GizmoElementSphere.cpp b/Engine/Source/Runtime/InteractiveToolsFramework/Private/BaseGizmos/GizmoElementSphere.cpp
--- a/Engine/Source/Runtime/InteractiveToolsFramework/Private/BaseGizmos/GizmoElementSphere.cpp
+++ b/Engine/Source/Runtime/InteractiveToolsFramework/Private/BaseGizmos/GizmoElementSphere.cpp
@@ -23,7 +23,7 @@ void UGizmoElementSphere::Render(IToolsContextRenderAPI* RenderAPI, const FRende
 				PDI,
 				CurrentRenderState.LocalToWorldTransform.GetTranslation(),
 				CurrentRenderState.LocalToWorldTransform.Rotator(),
-				FVector::OneVector * Radius,
+				Radius * CurrentRenderState.LocalToWorldTransform.GetScale3D(),
 				NumSides,
 				NumSides,
 				UseMaterial->GetRenderProxy(), SDPG_Foreground);
