[ITF] Add two new states to Gizmo Elements:
- "Subdued" for when for example X is being manipulated, Y and Z are considered "Subdued"
- "Selected" for when an element is designated for interaction indirectly, but isn't currently

diff --git a/Engine/Source/Runtime/InteractiveToolsFramework/Private/BaseGizmos/GizmoElementBase.cpp b/Engine/Source/Runtime/InteractiveToolsFramework/Private/BaseGizmos/GizmoElementBase.cpp
--- a/Engine/Source/Runtime/InteractiveToolsFramework/Private/BaseGizmos/GizmoElementBase.cpp
+++ b/Engine/Source/Runtime/InteractiveToolsFramework/Private/BaseGizmos/GizmoElementBase.cpp
@@ -205,6 +205,12 @@ bool UGizmoElementBase::GetEnabledForInteractionState(const EGizmoElementInterac
 
 		case EGizmoElementInteractionState::Interacting:	
 			return bEnabledForInteractingState;
+
+		case EGizmoElementInteractionState::Selected:	
+			return bEnabledForSelectedState;
+
+		case EGizmoElementInteractionState::Subdued:	
+			return bEnabledForSubduedState;
 		}
 	}
 
@@ -387,6 +393,26 @@ bool UGizmoElementBase::GetEnabledForInteractingState()
 	return bEnabledForInteractingState;
 }
 
+void UGizmoElementBase::SetEnabledForSelectedState(bool bInEnabledForSelectedState)
+{
+	bEnabledForSelectedState = bInEnabledForSelectedState;
+}
+
+bool UGizmoElementBase::GetEnabledForSelectedState()
+{
+	return bEnabledForSelectedState;
+}
+
+void UGizmoElementBase::SetEnabledForSubduedState(bool bInEnabledForSubduedState)
+{
+	bEnabledForSubduedState = bInEnabledForSubduedState;
+}
+
+bool UGizmoElementBase::GetEnabledForSubduedState()
+{
+	return bEnabledForSubduedState;
+}
+
 void UGizmoElementBase::SetPartIdentifier(uint32 InPartIdentifier)
 {
 	PartIdentifier = InPartIdentifier;
@@ -616,6 +642,46 @@ void UGizmoElementBase::ClearInteractMaterial()
 	MeshRenderAttributes.InteractMaterial.Reset();
 }
 
+void UGizmoElementBase::SetSelectMaterial(TWeakObjectPtr<UMaterialInterface> InSelectMaterial, bool InOverridesChildState)
+{
+	MeshRenderAttributes.SelectMaterial.SetMaterial(InSelectMaterial, InOverridesChildState);
+}
+
+const UMaterialInterface* UGizmoElementBase::GetSelectMaterial() const
+{
+	return MeshRenderAttributes.SelectMaterial.GetMaterial();
+}
+
+bool UGizmoElementBase::DoesSelectMaterialOverrideChildState() const
+{
+	return MeshRenderAttributes.SelectMaterial.bOverridesChildState;
+}
+
+void UGizmoElementBase::ClearSelectMaterial()
+{
+	MeshRenderAttributes.SelectMaterial.Reset();
+}
+
+void UGizmoElementBase::SetSubdueMaterial(TWeakObjectPtr<UMaterialInterface> InSubdueMaterial, bool InOverridesChildState)
+{
+	MeshRenderAttributes.SubdueMaterial.SetMaterial(InSubdueMaterial, InOverridesChildState);
+}
+
+const UMaterialInterface* UGizmoElementBase::GetSubdueMaterial() const
+{
+	return MeshRenderAttributes.SubdueMaterial.GetMaterial();
+}
+
+bool UGizmoElementBase::DoesSubdueMaterialOverrideChildState() const
+{
+	return MeshRenderAttributes.SubdueMaterial.bOverridesChildState;
+}
+
+void UGizmoElementBase::ClearSubdueMaterial()
+{
+	MeshRenderAttributes.SubdueMaterial.Reset();
+}
+
 void UGizmoElementBase::SetVertexColor(FLinearColor InVertexColor, bool InOverridesChildState)
 {
 	MeshRenderAttributes.VertexColor.SetColor(InVertexColor, InOverridesChildState);
@@ -687,3 +753,53 @@ void UGizmoElementBase::ClearInteractVertexColor()
 {
 	MeshRenderAttributes.InteractVertexColor.Reset();
 }
+
+void UGizmoElementBase::SetSelectVertexColor(FLinearColor InVertexColor, bool InOverridesChildState)
+{
+	MeshRenderAttributes.SelectVertexColor.SetColor(InVertexColor, InOverridesChildState);
+}
+
+FLinearColor UGizmoElementBase::GetSelectVertexColor() const
+{
+	return MeshRenderAttributes.SelectVertexColor.GetColor();
+}
+
+bool UGizmoElementBase::HasSelectVertexColor() const
+{
+	return MeshRenderAttributes.SelectVertexColor.bHasValue;
+}
+
+bool UGizmoElementBase::DoesSelectVertexColorOverrideChildState() const
+{
+	return MeshRenderAttributes.SelectVertexColor.bOverridesChildState;
+}
+
+void UGizmoElementBase::ClearSelectVertexColor()
+{
+	MeshRenderAttributes.SelectVertexColor.Reset();
+}
+
+void UGizmoElementBase::SetSubdueVertexColor(FLinearColor InVertexColor, bool InOverridesChildState)
+{
+	MeshRenderAttributes.SubdueVertexColor.SetColor(InVertexColor, InOverridesChildState);
+}
+
+FLinearColor UGizmoElementBase::GetSubdueVertexColor() const
+{
+	return MeshRenderAttributes.SubdueVertexColor.GetColor();
+}
+
+bool UGizmoElementBase::HasSubdueVertexColor() const
+{
+	return MeshRenderAttributes.SubdueVertexColor.bHasValue;
+}
+
+bool UGizmoElementBase::DoesSubdueVertexColorOverrideChildState() const
+{
+	return MeshRenderAttributes.SubdueVertexColor.bOverridesChildState;
+}
+
+void UGizmoElementBase::ClearSubdueVertexColor()
+{
+	MeshRenderAttributes.SubdueVertexColor.Reset();
+}
diff --git a/Engine/Source/Runtime/InteractiveToolsFramework/Private/BaseGizmos/GizmoElementLineBase.cpp b/Engine/Source/Runtime/InteractiveToolsFramework/Private/BaseGizmos/GizmoElementLineBase.cpp
--- a/Engine/Source/Runtime/InteractiveToolsFramework/Private/BaseGizmos/GizmoElementLineBase.cpp
+++ b/Engine/Source/Runtime/InteractiveToolsFramework/Private/BaseGizmos/GizmoElementLineBase.cpp
@@ -17,17 +17,32 @@ float UGizmoElementLineBase::GetCurrentLineThickness(bool bPerspectiveView, floa
 {
 	float CurrentLineThickness;
 
-	if (ElementInteractionState == EGizmoElementInteractionState::Hovering)
+	auto GetMultipliedLineThickness = [this](float InMultiplier) -> float
 	{
-		CurrentLineThickness = (LineThickness > 0.0f ? LineThickness * HoverLineThicknessMultiplier : HoverLineThicknessMultiplier);
-	}
-	else if (ElementInteractionState == EGizmoElementInteractionState::Interacting)
-	{
-		CurrentLineThickness = (LineThickness > 0.0f ? LineThickness * InteractLineThicknessMultiplier : InteractLineThicknessMultiplier);
-	}
-	else
+		return (LineThickness > UE_SMALL_NUMBER ? LineThickness * InMultiplier : InMultiplier);
+	};
+
+	switch (ElementInteractionState)
 	{
+	case EGizmoElementInteractionState::Hovering:
+		CurrentLineThickness = GetMultipliedLineThickness(HoverLineThicknessMultiplier);
+		break;
+
+	case EGizmoElementInteractionState::Interacting:
+		CurrentLineThickness = GetMultipliedLineThickness(InteractLineThicknessMultiplier);
+		break;
+
+	case EGizmoElementInteractionState::Selected:
+		CurrentLineThickness = GetMultipliedLineThickness(SelectLineThicknessMultiplier);
+		break;
+
+	case EGizmoElementInteractionState::Subdued:
+		CurrentLineThickness = GetMultipliedLineThickness(SubdueLineThicknessMultiplier);
+		break;
+
+	default:
 		CurrentLineThickness = LineThickness;
+		break;
 	}
 
 	if (bPerspectiveView)
@@ -76,6 +91,26 @@ float UGizmoElementLineBase::GetInteractLineThicknessMultiplier() const
 	return InteractLineThicknessMultiplier;
 }
 
+void UGizmoElementLineBase::SetSelectLineThicknessMultiplier(float InSelectLineThicknessMultiplier)
+{
+	SelectLineThicknessMultiplier = InSelectLineThicknessMultiplier;
+}
+
+float UGizmoElementLineBase::GetSelectLineThicknessMultiplier() const
+{
+	return SelectLineThicknessMultiplier;
+}
+
+void UGizmoElementLineBase::SetSubdueLineThicknessMultiplier(float InSubdueLineThicknessMultiplier)
+{
+	SubdueLineThicknessMultiplier = InSubdueLineThicknessMultiplier;
+}
+
+float UGizmoElementLineBase::GetSubdueLineThicknessMultiplier() const
+{
+	return SubdueLineThicknessMultiplier;
+}
+
 void UGizmoElementLineBase::SetScreenSpaceLine(bool bInScreenSpaceLine)
 {
 	bScreenSpaceLine = bInScreenSpaceLine;
@@ -159,5 +194,52 @@ void UGizmoElementLineBase::ClearInteractLineColor()
 	LineRenderAttributes.InteractLineColor.Reset();
 }
 
+void UGizmoElementLineBase::SetSelectLineColor(FLinearColor InColor, bool InOverridesChildState)
+{
+	LineRenderAttributes.SelectLineColor.SetColor(InColor, InOverridesChildState);
+}
+
+FLinearColor UGizmoElementLineBase::GetSelectLineColor() const
+{
+	return LineRenderAttributes.SelectLineColor.GetColor();
+}
+
+bool UGizmoElementLineBase::HasSelectLineColor() const
+{
+	return LineRenderAttributes.SelectLineColor.bHasValue;
+}
+
+bool UGizmoElementLineBase::DoesSelectLineColorOverrideChildState() const
+{
+	return LineRenderAttributes.SelectLineColor.bOverridesChildState;
+}
+
+void UGizmoElementLineBase::ClearSelectLineColor()
+{
+	LineRenderAttributes.SelectLineColor.Reset();
+}
+
+void UGizmoElementLineBase::SetSubdueLineColor(FLinearColor InColor, bool InOverridesChildState)
+{
+	LineRenderAttributes.SubdueLineColor.SetColor(InColor, InOverridesChildState);
+}
+
+FLinearColor UGizmoElementLineBase::GetSubdueLineColor() const
+{
+	return LineRenderAttributes.SubdueLineColor.GetColor();
+}
+
+bool UGizmoElementLineBase::HasSubdueLineColor() const
+{
+	return LineRenderAttributes.SubdueLineColor.bHasValue;
+}
 
+bool UGizmoElementLineBase::DoesSubdueLineColorOverrideChildState() const
+{
+	return LineRenderAttributes.SubdueLineColor.bOverridesChildState;
+}
 
+void UGizmoElementLineBase::ClearSubdueLineColor()
+{
+	LineRenderAttributes.SubdueLineColor.Reset();
+}
diff --git a/Engine/Source/Runtime/InteractiveToolsFramework/Private/BaseGizmos/GizmoElementRenderState.cpp b/Engine/Source/Runtime/InteractiveToolsFramework/Private/BaseGizmos/GizmoElementRenderState.cpp
--- a/Engine/Source/Runtime/InteractiveToolsFramework/Private/BaseGizmos/GizmoElementRenderState.cpp
+++ b/Engine/Source/Runtime/InteractiveToolsFramework/Private/BaseGizmos/GizmoElementRenderState.cpp
@@ -70,13 +70,32 @@ void FGizmoElementMaterialAttribute::UpdateState(const FGizmoElementMaterialAttr
 
 const UMaterialInterface* FGizmoElementMeshRenderStateAttributes::GetMaterial(EGizmoElementInteractionState InteractionState)
 {
-	if (InteractionState == EGizmoElementInteractionState::Hovering)
+	return GetMaterialInternal(InteractionState);
+}
+
+const UMaterialInterface* FGizmoElementMeshRenderStateAttributes::GetMaterial(EGizmoElementInteractionState InteractionState) const
+{
+	return GetMaterialInternal(InteractionState);
+}
+
+const UMaterialInterface* FGizmoElementMeshRenderStateAttributes::GetMaterialInternal(EGizmoElementInteractionState InteractionState) const
+{
+	switch (InteractionState)
 	{
+	case EGizmoElementInteractionState::Hovering:
 		return HoverMaterial.GetMaterial();
-	}
-	else if (InteractionState == EGizmoElementInteractionState::Interacting)
-	{
+
+	case EGizmoElementInteractionState::Interacting:
 		return InteractMaterial.GetMaterial();
+
+	case EGizmoElementInteractionState::Selected:
+		return SelectMaterial.GetMaterial() ? SelectMaterial.GetMaterial() : InteractMaterial.GetMaterial();
+
+	case EGizmoElementInteractionState::Subdued:
+		return SubdueMaterial.GetMaterial() ? SubdueMaterial.GetMaterial() : Material.GetMaterial();
+
+	default:
+		break;
 	}
 
 	return Material.GetMaterial();
@@ -84,14 +103,34 @@ const UMaterialInterface* FGizmoElementMeshRenderStateAttributes::GetMaterial(EG
 
 FLinearColor FGizmoElementMeshRenderStateAttributes::GetVertexColor(EGizmoElementInteractionState InteractionState)
 {
-	if (InteractionState == EGizmoElementInteractionState::Hovering)
+	return GetVertexColorInternal(InteractionState);
+}
+
+FLinearColor FGizmoElementMeshRenderStateAttributes::GetVertexColor(EGizmoElementInteractionState InteractionState) const
+{
+	return GetVertexColorInternal(InteractionState);
+}
+
+FLinearColor FGizmoElementMeshRenderStateAttributes::GetVertexColorInternal(EGizmoElementInteractionState InteractionState) const
+{
+	switch (InteractionState)
 	{
+	case EGizmoElementInteractionState::Hovering:
 		return HoverVertexColor.GetColor();
-	}
-	else if (InteractionState == EGizmoElementInteractionState::Interacting)
-	{
+
+	case EGizmoElementInteractionState::Interacting:
 		return InteractVertexColor.GetColor();
+
+	case EGizmoElementInteractionState::Selected:
+		return SelectVertexColor.bHasValue ? SelectVertexColor.GetColor() : InteractVertexColor.GetColor();
+
+	case EGizmoElementInteractionState::Subdued:
+		return SubdueVertexColor.bHasValue ? SubdueVertexColor.GetColor() : VertexColor.GetColor();
+
+	default:
+		break;
 	}
+
 	return VertexColor.GetColor();
 }
 
@@ -100,20 +139,34 @@ void FGizmoElementMeshRenderStateAttributes::Update(FGizmoElementMeshRenderState
 	Material.UpdateState(InChildAttributes.Material);
 	HoverMaterial.UpdateState(InChildAttributes.HoverMaterial);
 	InteractMaterial.UpdateState(InChildAttributes.InteractMaterial);
+	SelectMaterial.UpdateState(InChildAttributes.SelectMaterial);
+	SubdueMaterial.UpdateState(InChildAttributes.SubdueMaterial);
+
 	VertexColor.UpdateState(InChildAttributes.VertexColor);
 	HoverVertexColor.UpdateState(InChildAttributes.HoverVertexColor);
 	InteractVertexColor.UpdateState(InChildAttributes.InteractVertexColor);
+	SelectVertexColor.UpdateState(InChildAttributes.SelectVertexColor);
+	SubdueVertexColor.UpdateState(InChildAttributes.SubdueVertexColor);
 }
 
 FLinearColor FGizmoElementLineRenderStateAttributes::GetLineColor(EGizmoElementInteractionState InteractionState)
 {
-	if (InteractionState == EGizmoElementInteractionState::Hovering)
+	switch (InteractionState)
 	{
+	case EGizmoElementInteractionState::Hovering:
 		return HoverLineColor.GetColor();
-	}
-	else if (InteractionState == EGizmoElementInteractionState::Interacting)
-	{
+
+	case EGizmoElementInteractionState::Interacting:
 		return InteractLineColor.GetColor();
+
+	case EGizmoElementInteractionState::Selected:
+		return SelectLineColor.bHasValue ? SelectLineColor.GetColor() : InteractLineColor.GetColor();
+
+	case EGizmoElementInteractionState::Subdued:
+		return SubdueLineColor.bHasValue ? SubdueLineColor.GetColor() : LineColor.GetColor();
+
+	default:
+		break;
 	}
 
 	return LineColor.GetColor();
@@ -124,4 +177,6 @@ void FGizmoElementLineRenderStateAttributes::Update(FGizmoElementLineRenderState
 	LineColor.UpdateState(InChildAttributes.LineColor);
 	HoverLineColor.UpdateState(InChildAttributes.HoverLineColor);
 	InteractLineColor.UpdateState(InChildAttributes.InteractLineColor);
+	SelectLineColor.UpdateState(InChildAttributes.SelectLineColor);
+	SubdueLineColor.UpdateState(InChildAttributes.SubdueLineColor);
 }
diff --git a/Engine/Source/Runtime/InteractiveToolsFramework/Public/BaseGizmos/GizmoElementBase.h b/Engine/Source/Runtime/InteractiveToolsFramework/Public/BaseGizmos/GizmoElementBase.h
--- a/Engine/Source/Runtime/InteractiveToolsFramework/Public/BaseGizmos/GizmoElementBase.h
+++ b/Engine/Source/Runtime/InteractiveToolsFramework/Public/BaseGizmos/GizmoElementBase.h
@@ -155,6 +155,14 @@ public:
 	INTERACTIVETOOLSFRAMEWORK_API virtual void SetEnabledForInteractingState(bool bInEnabledForInteractingState);
 	INTERACTIVETOOLSFRAMEWORK_API virtual bool GetEnabledForInteractingState();
 
+	// Whether element is enabled when element state is selected.
+	INTERACTIVETOOLSFRAMEWORK_API virtual void SetEnabledForSelectedState(bool bInEnabledForSelectedState);
+	INTERACTIVETOOLSFRAMEWORK_API virtual bool GetEnabledForSelectedState();
+
+	// Whether element is enabled when element state is subdued.
+	INTERACTIVETOOLSFRAMEWORK_API virtual void SetEnabledForSubduedState(bool bInEnabledForSubduedState);
+	INTERACTIVETOOLSFRAMEWORK_API virtual bool GetEnabledForSubduedState();
+
 	// For an element hierarchy representing multiple parts of a single gizmo, the part identifier establishes 
 	// a correspondence between a gizmo part and the elements that represent that part. The recognized
 	// part identifier values should be defined in the gizmo. Gizmo part identifiers must be greater than or 
@@ -167,7 +175,7 @@ public:
 	INTERACTIVETOOLSFRAMEWORK_API virtual void SetElementState(EGizmoElementState InElementState);
 	INTERACTIVETOOLSFRAMEWORK_API virtual EGizmoElementState GetElementState() const;
 
-	// Object interaction state - None, Hovering or Interacting
+	// Object interaction state - None, Hovering, Interacting, Selected or Subdued
 	INTERACTIVETOOLSFRAMEWORK_API virtual void SetElementInteractionState(EGizmoElementInteractionState InInteractionState);
 	INTERACTIVETOOLSFRAMEWORK_API virtual EGizmoElementInteractionState GetElementInteractionState() const;
 
@@ -231,7 +239,7 @@ public:
 	INTERACTIVETOOLSFRAMEWORK_API virtual float GetPixelHitDistanceThreshold() const;
 
 	//
-	// Methods for managing render state attributes: Material, HoverMaterial, InteractMaterial, VertexColor 
+	// Methods for managing render state attributes: Material, HoverMaterial, InteractMaterial, SelectMaterial, SubdueMaterial, VertexColor 
 	// 
 	// State inheritance works as follows: 
 	// - Gizmo element state that is not set inherits from the corresponding state in the current render traversal.
@@ -268,7 +276,7 @@ public:
 	INTERACTIVETOOLSFRAMEWORK_API virtual void ClearHoverMaterial();
 
 	// Set mesh render state interact material attribute. 
-	//  @param InHoverMaterial - interact material to be set
+	//  @param InInteractMaterial - interact material to be set
 	//  @param InOverridesChildState - when true, this interact material will override the material of all child elements.
 	INTERACTIVETOOLSFRAMEWORK_API virtual void SetInteractMaterial(TWeakObjectPtr<UMaterialInterface> InInteractMaterial, bool InOverridesChildState = false);
 
@@ -281,6 +289,34 @@ public:
 	// Clear mesh render interact state material attribute. 
 	INTERACTIVETOOLSFRAMEWORK_API virtual void ClearInteractMaterial();
 
+	// Set mesh render state Select material attribute. 
+	//  @param InSelectMaterial - Select material to be set
+	//  @param InOverridesChildState - when true, this Select material will override the material of all child elements.
+	INTERACTIVETOOLSFRAMEWORK_API virtual void SetSelectMaterial(TWeakObjectPtr<UMaterialInterface> InSelectMaterial, bool InOverridesChildState = false);
+
+	// Get mesh render state Select material attribute's value. 
+	INTERACTIVETOOLSFRAMEWORK_API virtual const UMaterialInterface* GetSelectMaterial() const;
+
+	// Get mesh render state Select material attribute's override setting. 
+	INTERACTIVETOOLSFRAMEWORK_API virtual bool DoesSelectMaterialOverrideChildState() const;
+
+	// Clear mesh render Select state material attribute. 
+	INTERACTIVETOOLSFRAMEWORK_API virtual void ClearSelectMaterial();
+
+	// Set mesh render state Subdue material attribute. 
+	//  @param InSubdueMaterial - Subdue material to be set
+	//  @param InOverridesChildState - when true, this Subdue material will override the material of all child elements.
+	INTERACTIVETOOLSFRAMEWORK_API virtual void SetSubdueMaterial(TWeakObjectPtr<UMaterialInterface> InSubdueMaterial, bool InOverridesChildState = false);
+
+	// Get mesh render state Subdue material attribute's value. 
+	INTERACTIVETOOLSFRAMEWORK_API virtual const UMaterialInterface* GetSubdueMaterial() const;
+
+	// Get mesh render state Subdue material attribute's override setting. 
+	INTERACTIVETOOLSFRAMEWORK_API virtual bool DoesSubdueMaterialOverrideChildState() const;
+
+	// Clear mesh render Subdue state material attribute. 
+	INTERACTIVETOOLSFRAMEWORK_API virtual void ClearSubdueMaterial();
+
 	// Set mesh render state vertex color attribute. 
 	//  @param InVertexColor - vertex color to be set
 	//  @param InOverridesChildState - when true, this vertex color will override the material of all child elements.
@@ -332,6 +368,40 @@ public:
 	// Clear mesh render state vertex color attribute.
 	INTERACTIVETOOLSFRAMEWORK_API virtual void ClearInteractVertexColor();
 
+	// Set mesh render state vertex color attribute. 
+	//  @param InVertexColor - vertex color to be set
+	//  @param InOverridesChildState - when true, this vertex color will override the material of all child elements.
+	INTERACTIVETOOLSFRAMEWORK_API virtual void SetSelectVertexColor(FLinearColor InVertexColor, bool InOverridesChildState = false);
+
+	// Get mesh render state vertex color attribute's value. 
+	INTERACTIVETOOLSFRAMEWORK_API virtual FLinearColor GetSelectVertexColor() const;
+
+	// Returns true, if mesh render state vertex color attribute has been set. 
+	INTERACTIVETOOLSFRAMEWORK_API virtual bool HasSelectVertexColor() const;
+
+	// Get mesh render state vertex color attribute's override setting. 
+	INTERACTIVETOOLSFRAMEWORK_API virtual bool DoesSelectVertexColorOverrideChildState() const;
+
+	// Clear mesh render state vertex color attribute.
+	INTERACTIVETOOLSFRAMEWORK_API virtual void ClearSelectVertexColor();
+
+	// Set mesh render state vertex color attribute. 
+	//  @param InVertexColor - vertex color to be set
+	//  @param InOverridesChildState - when true, this vertex color will override the material of all child elements.
+	INTERACTIVETOOLSFRAMEWORK_API virtual void SetSubdueVertexColor(FLinearColor InVertexColor, bool InOverridesChildState = false);
+
+	// Get mesh render state vertex color attribute's value. 
+	INTERACTIVETOOLSFRAMEWORK_API virtual FLinearColor GetSubdueVertexColor() const;
+
+	// Returns true, if mesh render state vertex color attribute has been set. 
+	INTERACTIVETOOLSFRAMEWORK_API virtual bool HasSubdueVertexColor() const;
+
+	// Get mesh render state vertex color attribute's override setting. 
+	INTERACTIVETOOLSFRAMEWORK_API virtual bool DoesSubdueVertexColorOverrideChildState() const;
+
+	// Clear mesh render state vertex color attribute.
+	INTERACTIVETOOLSFRAMEWORK_API virtual void ClearSubdueVertexColor();
+
 protected:
 
 	// Whether element is enabled. Render and LineTrace only occur when bEnabled is true.
@@ -358,6 +428,14 @@ protected:
 	UPROPERTY()
 	bool bEnabledForInteractingState = true;
 
+	// Whether element is enabled when element state is selected
+	UPROPERTY()
+	bool bEnabledForSelectedState = true;
+
+	// Whether element is enabled when element state is subdued
+	UPROPERTY()
+	bool bEnabledForSubduedState = true;
+
 	// Part identifier
 	UPROPERTY()
 	uint32 PartIdentifier = DefaultPartIdentifier;
@@ -370,7 +448,7 @@ protected:
 	UPROPERTY()
 	EGizmoElementState ElementState = EGizmoElementState::VisibleAndHittable;
 
-	// Current element interaction state - None, Hovering or Interacting
+	// Current element interaction state - None, Hovering, Interacting, Selected or Subdued
 	UPROPERTY()
 	EGizmoElementInteractionState ElementInteractionState = EGizmoElementInteractionState::None;
 
diff --git a/Engine/Source/Runtime/InteractiveToolsFramework/Public/BaseGizmos/GizmoElementLineBase.h b/Engine/Source/Runtime/InteractiveToolsFramework/Public/BaseGizmos/GizmoElementLineBase.h
--- a/Engine/Source/Runtime/InteractiveToolsFramework/Public/BaseGizmos/GizmoElementLineBase.h
+++ b/Engine/Source/Runtime/InteractiveToolsFramework/Public/BaseGizmos/GizmoElementLineBase.h
@@ -34,12 +34,20 @@ public:
 	INTERACTIVETOOLSFRAMEWORK_API virtual void SetInteractLineThicknessMultiplier(float InInteractLineThicknessMultiplier);
 	INTERACTIVETOOLSFRAMEWORK_API virtual float GetInteractLineThicknessMultiplier() const;
 
+	// Multiplier applied to line thickness when selected
+	INTERACTIVETOOLSFRAMEWORK_API virtual void SetSelectLineThicknessMultiplier(float InSelectLineThicknessMultiplier);
+	INTERACTIVETOOLSFRAMEWORK_API virtual float GetSelectLineThicknessMultiplier() const;
+
+	// Multiplier applied to line thickness when subdued
+	INTERACTIVETOOLSFRAMEWORK_API virtual void SetSubdueLineThicknessMultiplier(float InSubdueLineThicknessMultiplier);
+	INTERACTIVETOOLSFRAMEWORK_API virtual float GetSubdueLineThicknessMultiplier() const;
+
 	// Whether line thickness is in screen space 
 	INTERACTIVETOOLSFRAMEWORK_API virtual void SetScreenSpaceLine(bool bInScreenSpaceLine);
 	INTERACTIVETOOLSFRAMEWORK_API virtual bool GetScreenSpaceLine() const;
 
 	//
-	// Methods for managing line state attributes: LineColor, HoverLineColor, InteractLineColor
+	// Methods for managing line state attributes: LineColor, HoverLineColor, InteractLineColor, SelectLineColor, SubdueLineColor
 	//
 	// State inheritance works as follows:
 	// - Gizmo element state that is not set inherits from the corresponding state in the current render traversal.
@@ -98,6 +106,40 @@ public:
 	// Clear line render state interact line color attribute.
 	INTERACTIVETOOLSFRAMEWORK_API virtual void ClearInteractLineColor();
 
+	// Set line render state Select line color attribute. 
+	//  @param InColor - Select line color to be set
+	//  @param InOverridesChildState - when true, this Select line color will override the line color of all child elements.
+	INTERACTIVETOOLSFRAMEWORK_API virtual void SetSelectLineColor(FLinearColor InColor, bool InOverridesChildState = false);
+
+	// Get line render state interact line color attribute's value. 
+	INTERACTIVETOOLSFRAMEWORK_API virtual FLinearColor GetSelectLineColor() const;
+
+	// Returns true if line render state interact line color attribute has been set. 
+	INTERACTIVETOOLSFRAMEWORK_API virtual bool HasSelectLineColor() const;
+
+	// Get line render state interact line color attribute's override setting. 
+	INTERACTIVETOOLSFRAMEWORK_API virtual bool DoesSelectLineColorOverrideChildState() const;
+
+	// Clear line render state interact line color attribute.
+	INTERACTIVETOOLSFRAMEWORK_API virtual void ClearSelectLineColor();
+
+	// Set line render state subdue line color attribute. 
+	//  @param InColor - subdue line color to be set
+	//  @param InOverridesChildState - when true, this subdue line color will override the line color of all child elements.
+	INTERACTIVETOOLSFRAMEWORK_API virtual void SetSubdueLineColor(FLinearColor InColor, bool InOverridesChildState = false);
+
+	// Get line render state interact line color attribute's value. 
+	INTERACTIVETOOLSFRAMEWORK_API virtual FLinearColor GetSubdueLineColor() const;
+
+	// Returns true if line render state interact line color attribute has been set. 
+	INTERACTIVETOOLSFRAMEWORK_API virtual bool HasSubdueLineColor() const;
+
+	// Get line render state interact line color attribute's override setting. 
+	INTERACTIVETOOLSFRAMEWORK_API virtual bool DoesSubdueLineColorOverrideChildState() const;
+
+	// Clear line render state interact line color attribute.
+	INTERACTIVETOOLSFRAMEWORK_API virtual void ClearSubdueLineColor();
+
 protected:
 
 	// Line render state attributes for this element
@@ -120,6 +162,14 @@ protected:
 	UPROPERTY(EditAnywhere, Category = Options)
 	float InteractLineThicknessMultiplier = 2.0f;
 
+	// Multiplier applied to line thickness when selected
+	UPROPERTY(EditAnywhere, Category = Options)
+	float SelectLineThicknessMultiplier = 2.0f;
+
+	// Multiplier applied to line thickness when subdued
+	UPROPERTY(EditAnywhere, Category = Options)
+	float SubdueLineThicknessMultiplier = 1.0f;
+
 	// Update render state during render traversal, determines the current render state for this element
 	// @return view dependent visibility, true if this element is visible in the current view.
 	INTERACTIVETOOLSFRAMEWORK_API virtual bool UpdateRenderState(IToolsContextRenderAPI* RenderAPI, const FVector& InLocalOrigin, FRenderTraversalState& InOutRenderState) override;
diff --git a/Engine/Source/Runtime/InteractiveToolsFramework/Public/BaseGizmos/GizmoElementRenderState.h b/Engine/Source/Runtime/InteractiveToolsFramework/Public/BaseGizmos/GizmoElementRenderState.h
--- a/Engine/Source/Runtime/InteractiveToolsFramework/Public/BaseGizmos/GizmoElementRenderState.h
+++ b/Engine/Source/Runtime/InteractiveToolsFramework/Public/BaseGizmos/GizmoElementRenderState.h
@@ -102,6 +102,14 @@ struct FGizmoElementMeshRenderStateAttributes
 	UPROPERTY()
 	FGizmoElementMaterialAttribute InteractMaterial;
 
+	// Select material
+	UPROPERTY()
+	FGizmoElementMaterialAttribute SelectMaterial;
+
+	// Subdue material
+	UPROPERTY()
+	FGizmoElementMaterialAttribute SubdueMaterial;
+
 	// Mesh vertex color
 	UPROPERTY()
 	FGizmoElementColorAttribute VertexColor;
@@ -114,14 +122,35 @@ struct FGizmoElementMeshRenderStateAttributes
 	UPROPERTY()
 	FGizmoElementColorAttribute InteractVertexColor;
 
+	// Select mesh vertex color
+	UPROPERTY()
+	FGizmoElementColorAttribute SelectVertexColor;
+
+	// Subdue mesh vertex color
+	UPROPERTY()
+	FGizmoElementColorAttribute SubdueVertexColor;
+
 	// Returns the material corresponding to the input interaction state
 	INTERACTIVETOOLSFRAMEWORK_API const UMaterialInterface* GetMaterial(EGizmoElementInteractionState InteractionState);
 
+	// Returns the material corresponding to the input interaction state
+	INTERACTIVETOOLSFRAMEWORK_API const UMaterialInterface* GetMaterial(EGizmoElementInteractionState InteractionState) const;
+
 	// Returns mesh vertex color
 	INTERACTIVETOOLSFRAMEWORK_API FLinearColor GetVertexColor(EGizmoElementInteractionState InteractionState);
 
+	// Returns mesh vertex color
+	INTERACTIVETOOLSFRAMEWORK_API FLinearColor GetVertexColor(EGizmoElementInteractionState InteractionState) const;
+
 	// Update this mesh render state based on a child mesh render state attribute.
 	INTERACTIVETOOLSFRAMEWORK_API void Update(FGizmoElementMeshRenderStateAttributes& InChildAttributes);
+
+private:
+	// Returns the material corresponding to the input interaction state
+	const UMaterialInterface* GetMaterialInternal(EGizmoElementInteractionState InteractionState) const;
+
+	// Returns mesh vertex color
+	FLinearColor GetVertexColorInternal(EGizmoElementInteractionState InteractionState) const;
 };
 
 //
@@ -144,6 +173,14 @@ struct FGizmoElementLineRenderStateAttributes
 	UPROPERTY()
 	FGizmoElementColorAttribute InteractLineColor;
 
+	// Select line color
+	UPROPERTY()
+	FGizmoElementColorAttribute SelectLineColor;
+
+	// Subdue line color
+	UPROPERTY()
+	FGizmoElementColorAttribute SubdueLineColor;
+
 	// Returns the line color corresponding to the input interaction state
 	INTERACTIVETOOLSFRAMEWORK_API FLinearColor GetLineColor(EGizmoElementInteractionState InteractionState);
 
diff --git a/Engine/Source/Runtime/InteractiveToolsFramework/Public/BaseGizmos/GizmoElementShared.h b/Engine/Source/Runtime/InteractiveToolsFramework/Public/BaseGizmos/GizmoElementShared.h
--- a/Engine/Source/Runtime/InteractiveToolsFramework/Public/BaseGizmos/GizmoElementShared.h
+++ b/Engine/Source/Runtime/InteractiveToolsFramework/Public/BaseGizmos/GizmoElementShared.h
@@ -28,7 +28,11 @@ enum class EGizmoElementInteractionState
 {
 	None,
 	Hovering,
-	Interacting
+	Interacting,
+	Selected,		// When an element is "Selected", it is considered to be the target of Interaction via indirect manipulation (not hit tested).
+	Subdued,		// When one or more other elements considered part of a set is active, but this one is not (ie. X Axis is Interacting, Y and Z are Subdued).
+
+	Max
 };
 
 //
