[ITF] Bumped DepthBias on various PDI DrawLine* calls (within UGizmoElements) to avoid z-fighting

diff --git a/Engine/Source/Runtime/InteractiveToolsFramework/Private/BaseGizmos/GizmoElementCircle.cpp b/Engine/Source/Runtime/InteractiveToolsFramework/Private/BaseGizmos/GizmoElementCircle.cpp
--- a/Engine/Source/Runtime/InteractiveToolsFramework/Private/BaseGizmos/GizmoElementCircle.cpp
+++ b/Engine/Source/Runtime/InteractiveToolsFramework/Private/BaseGizmos/GizmoElementCircle.cpp
@@ -12,6 +12,11 @@
 
 void UGizmoElementCircle::Render(IToolsContextRenderAPI* RenderAPI, const FRenderTraversalState& RenderState)
 {
+	if (!ensure(RenderAPI))
+	{
+		return;
+	}
+	
 	FRenderTraversalState CurrentRenderState(RenderState);
 	bool bVisibleViewDependent = UpdateRenderState(RenderAPI, Center, CurrentRenderState);
 
@@ -35,12 +40,18 @@ void UGizmoElementCircle::Render(IToolsContextRenderAPI* RenderAPI, const FRende
 
 		if (bDrawLine)
 		{
-			check(RenderAPI);
 			const FSceneView* View = RenderAPI->GetSceneView();
-			check(View);
+			if (!ensure(View))
+			{
+				return;
+			}
+
 			float CurrentLineThickness = GetCurrentLineThickness(View->IsPerspectiveProjection(), View->FOV);
 
-			DrawCircle(PDI, WorldCenter, WorldAxis0, WorldAxis1, CurrentRenderState.GetCurrentLineColor(), WorldRadius, NumSegments, SDPG_Foreground, CurrentLineThickness, 0.0, bScreenSpaceLine);
+			// Boost depth bias for lines to avoid z-fighting with the fill (of this or other elements)
+			const float LineDepthBias = 1.0f + CurrentLineThickness;
+
+			DrawCircle(PDI, WorldCenter, WorldAxis0, WorldAxis1, CurrentRenderState.GetCurrentLineColor(), WorldRadius, NumSegments, SDPG_Foreground, CurrentLineThickness, LineDepthBias, bScreenSpaceLine);
 		}
 	}
 }
diff --git a/Engine/Source/Runtime/InteractiveToolsFramework/Private/BaseGizmos/GizmoElementLineStrip.cpp b/Engine/Source/Runtime/InteractiveToolsFramework/Private/BaseGizmos/GizmoElementLineStrip.cpp
--- a/Engine/Source/Runtime/InteractiveToolsFramework/Private/BaseGizmos/GizmoElementLineStrip.cpp
+++ b/Engine/Source/Runtime/InteractiveToolsFramework/Private/BaseGizmos/GizmoElementLineStrip.cpp
@@ -13,6 +13,11 @@
 
 void UGizmoElementLineStrip::Render(IToolsContextRenderAPI* RenderAPI, const FRenderTraversalState& RenderState)
 {
+	if (!ensure(RenderAPI))
+	{
+		return;
+	}
+	
 	FRenderTraversalState CurrentRenderState(RenderState);
 	const bool bVisibleViewDependent = UpdateRenderState(RenderAPI, Base, CurrentRenderState);
 
@@ -21,12 +26,17 @@ void UGizmoElementLineStrip::Render(IToolsContextRenderAPI* RenderAPI, const FRe
 		FPrimitiveDrawInterface* PDI = RenderAPI->GetPrimitiveDrawInterface();
 		ComputeProjectedVertices(RenderState.LocalToWorldTransform, true);
 
-		check(RenderAPI);
 		const FSceneView* View = RenderAPI->GetSceneView();
-		check(View);
+		if (!ensure(View))
+		{
+			return;
+		}
 
 		const float CurrentLineThickness = GetCurrentLineThickness(View->IsPerspectiveProjection(), View->FOV);
 		const FColor LineColor = CurrentRenderState.GetCurrentLineColor().ToFColor(false);
+
+		// Boost depth bias for lines to avoid z-fighting with the fill (of this or other elements)
+		const float LineDepthBias = 1.0f + CurrentLineThickness;
 		
 		if(bDrawLineStrip)
 		{
@@ -34,7 +44,7 @@ void UGizmoElementLineStrip::Render(IToolsContextRenderAPI* RenderAPI, const FRe
 			{
 				const FVector& A = ProjectedVertices[Index];
 				const FVector& B = ProjectedVertices[Index+1];
-				PDI->DrawLine(A, B, LineColor, SDPG_Foreground, CurrentLineThickness, 0.0, bScreenSpaceLine);
+				PDI->DrawLine(A, B, LineColor, SDPG_Foreground, CurrentLineThickness, LineDepthBias, bScreenSpaceLine);
 			}
 		}
 		else
@@ -43,7 +53,7 @@ void UGizmoElementLineStrip::Render(IToolsContextRenderAPI* RenderAPI, const FRe
 			{
 				const FVector& A = ProjectedVertices[Index];
 				const FVector& B = ProjectedVertices[Index+1];
-				PDI->DrawLine(A, B, LineColor, SDPG_Foreground, CurrentLineThickness, 0.0, bScreenSpaceLine);
+				PDI->DrawLine(A, B, LineColor, SDPG_Foreground, CurrentLineThickness, LineDepthBias, bScreenSpaceLine);
 			}
 		}
 	}
diff --git a/Engine/Source/Runtime/InteractiveToolsFramework/Private/BaseGizmos/GizmoElementRectangle.cpp b/Engine/Source/Runtime/InteractiveToolsFramework/Private/BaseGizmos/GizmoElementRectangle.cpp
--- a/Engine/Source/Runtime/InteractiveToolsFramework/Private/BaseGizmos/GizmoElementRectangle.cpp
+++ b/Engine/Source/Runtime/InteractiveToolsFramework/Private/BaseGizmos/GizmoElementRectangle.cpp
@@ -12,6 +12,11 @@
 
 void UGizmoElementRectangle::Render(IToolsContextRenderAPI* RenderAPI, const FRenderTraversalState& RenderState)
 {
+	if (!ensure(RenderAPI))
+	{
+		return;
+	}
+	
 	FRenderTraversalState CurrentRenderState(RenderState);
 	bool bVisibleViewDependent = UpdateRenderState(RenderAPI, Center, CurrentRenderState);
 
@@ -33,15 +38,22 @@ void UGizmoElementRectangle::Render(IToolsContextRenderAPI* RenderAPI, const FRe
 				DrawRectangleMesh(PDI, WorldCenter, WorldUpAxis, WorldSideAxis, VertexColor, WorldWidth, WorldHeight, UseMaterial->GetRenderProxy(), SDPG_Foreground);
 			}
 		}
+
 		if (bDrawLine)
 		{
-			check(RenderAPI);
 			const FSceneView* View = RenderAPI->GetSceneView();
-			check(View);
+			if (!ensure(View))
+			{
+				return;
+			}
+
 			float CurrentLineThickness = GetCurrentLineThickness(View->IsPerspectiveProjection(), View->FOV);
 
+			// Boost depth bias for lines to avoid z-fighting with the fill (of this or other elements)
+			const float LineDepthBias = 1.0f + CurrentLineThickness;
+
 			FColor LineColor = CurrentRenderState.GetCurrentLineColor().ToFColor(false);
-			DrawRectangle(PDI, WorldCenter, WorldUpAxis, WorldSideAxis, LineColor, WorldWidth, WorldHeight, SDPG_Foreground, CurrentLineThickness, 0.0, bScreenSpaceLine);
+			DrawRectangle(PDI, WorldCenter, WorldUpAxis, WorldSideAxis, LineColor, WorldWidth, WorldHeight, SDPG_Foreground, CurrentLineThickness, LineDepthBias, bScreenSpaceLine);
 		}
 	}
 }
