[ITF] GizmoElementArc end border lines and fixed normals

diff --git a/Engine/Source/Runtime/InteractiveToolsFramework/Private/BaseGizmos/GizmoElementArc.cpp b/Engine/Source/Runtime/InteractiveToolsFramework/Private/BaseGizmos/GizmoElementArc.cpp
--- a/Engine/Source/Runtime/InteractiveToolsFramework/Private/BaseGizmos/GizmoElementArc.cpp
+++ b/Engine/Source/Runtime/InteractiveToolsFramework/Private/BaseGizmos/GizmoElementArc.cpp
@@ -14,7 +14,7 @@
 
 static void DrawThickArc(FPrimitiveDrawInterface* PDI, const FVector& InCenter, const FVector& InAxis0, const FVector& InAxis1,
 	double InOuterRadius, double InInnerRadius, const double InStartAngle, const double InEndAngle, const int32 InNumSegments,
-	const FMaterialRenderProxy* MaterialRenderProxy, const FColor& InColor)
+	const FMaterialRenderProxy* MaterialRenderProxy, const FColor& InColor, const FColor& InLineColor, const double InLineThickness = 0.0)
 {
 	// Implementation copied from FWidget::DrawThickArc. This should eventually be moved to PrimDrawingUtils.h/cpp.
 
@@ -23,32 +23,44 @@ static void DrawThickArc(FPrimitiveDrawInterface* PDI, const FVector& InCenter,
 		return;
 	}
 
+	// Boost depth bias for lines to avoid z-fighting with the fill
+	const float LineDepthBias = 1.0f + InLineThickness;
+
 	const int32 NumPoints = FMath::TruncToInt32(InNumSegments * (InEndAngle - InStartAngle) / (UE_DOUBLE_PI / 2.0)) + 1;
 
 	FColor TriangleColor = InColor;
-	FColor RingColor = InColor;
-	RingColor.A = MAX_uint8;
 
-	FVector ZAxis = InAxis0 ^ InAxis1;
+	const FVector3f Axis1f = static_cast<FVector3f>(InAxis1);
+	const FVector3f Axis0f = static_cast<FVector3f>(InAxis0);
+
+	// Retain the higher precision version for rotation calculations
+	const FVector ZAxis = InAxis0 ^ InAxis1;
+	const FVector3f ZAxisf = Axis0f ^ Axis1f;
 	FVector LastWorldVertex;
 
+	// Whether to draw the lines on either end of the arc (the edges of the proverbial pizza slice).
+	// We don't draw them if the arc forms a complete/closed circle, so we check if the NumPoints exceeds the NumSegments in 360 degrees (NumSegments is per 90 degrees).
+	const bool bDrawEndLines = NumPoints <= (InNumSegments * 4);
+
+	TArray<FDynamicMeshVertex> Vertices;
+	Vertices.Reserve((NumPoints * 2) + 2);
+
 	FDynamicMeshBuilder MeshBuilder(PDI->View->GetFeatureLevel());
 
 	for (int32 RadiusIndex = 0; RadiusIndex < 2; ++RadiusIndex)
 	{
 		double Radius = (RadiusIndex == 0) ? InOuterRadius : InInnerRadius;
-		double TCRadius = Radius / (double)InOuterRadius;
+		double TCRadius = Radius / InOuterRadius;
 		//Compute vertices for base circle.
 		for (int32 VertexIndex = 0; VertexIndex <= NumPoints; VertexIndex++)
 		{
-			double Percent = VertexIndex / (double)NumPoints;
-			double Angle = FMath::Lerp(InStartAngle, InEndAngle, Percent);
+			double Percent = VertexIndex / static_cast<double>(NumPoints);
+			double Angle = FMath::Lerp(InEndAngle, InStartAngle, Percent);
 			double AngleDeg = FRotator::ClampAxis(Angle * 180.f / PI);
 
 			FVector VertexDir = InAxis0.RotateAngleAxis(AngleDeg, ZAxis);
 			VertexDir.Normalize();
 
-			double TCAngle = Percent * (PI / 2);
 			FVector2f TC(static_cast<float>(TCRadius * FMath::Cos(Angle)), static_cast<float>(TCRadius * FMath::Sin(Angle)));
 
 			// Keep the vertices in local space so that we don't lose precision when dealing with LWC
@@ -58,29 +70,46 @@ static void DrawThickArc(FPrimitiveDrawInterface* PDI, const FVector& InCenter,
 			Normal.Normalize();
 
 			FDynamicMeshVertex MeshVertex;
-			MeshVertex.Position = (FVector3f)VertexPosition;
+			MeshVertex.Position = static_cast<FVector3f>(VertexPosition);
 			MeshVertex.Color = TriangleColor;
 			MeshVertex.TextureCoordinate[0] = TC;
 
 			MeshVertex.SetTangents(
-				(FVector3f)-ZAxis,
-				FVector3f((-ZAxis) ^ Normal),
-				(FVector3f)Normal
+				Axis1f,
+				-Axis0f,
+				ZAxisf
 			);
 
-			MeshBuilder.AddVertex(MeshVertex); //Add bottom vertex
+			Vertices.Emplace(MeshVertex); //Add bottom vertex
 
 			// Push out the arc line borders so they dont z-fight with the mesh arcs
 			// DrawLine needs vertices in world space, but this is fine because it takes FVectors and works with LWC well
 			FVector EndLinePos = VertexPosition + InCenter;
 			if (VertexIndex != 0)
 			{
-				PDI->DrawLine(LastWorldVertex, EndLinePos, RingColor, SDPG_Foreground);
+				PDI->DrawTranslucentLine(LastWorldVertex, EndLinePos, InLineColor, SDPG_Foreground, InLineThickness, LineDepthBias);
 			}
 			LastWorldVertex = EndLinePos;
 		}
 	}
 
+	MeshBuilder.AddVertices(Vertices);
+
+	if (bDrawEndLines)
+	{
+		check(Vertices.Num() == (NumPoints * 2) + 2);
+
+		// From the first vert on the outer ring to the first vert on the inner ring
+		FVector FirstEndLinePosOuter = static_cast<FVector>(Vertices[0].Position) + InCenter;
+		FVector FirstEndLinePosInner = static_cast<FVector>(Vertices[NumPoints + 1].Position) + InCenter;
+		PDI->DrawTranslucentLine(FirstEndLinePosOuter, FirstEndLinePosInner, InLineColor, SDPG_Foreground, InLineThickness, LineDepthBias);
+
+		// From the last vert on the outer ring to the last vert on the inner ring
+		FVector LastEndLinePosOuter = static_cast<FVector>(Vertices[NumPoints].Position) + InCenter;
+		FVector LastEndLinePosInner = static_cast<FVector>(Vertices.Last().Position) + InCenter;
+		PDI->DrawTranslucentLine(LastEndLinePosOuter, LastEndLinePosInner, InLineColor, SDPG_Foreground, InLineThickness, LineDepthBias);
+	}
+
 	//Add top/bottom triangles, in the style of a fan.
 	int32 InnerVertexStartIndex = NumPoints + 1;
 	for (int32 VertexIndex = 0; VertexIndex < NumPoints; VertexIndex++)
@@ -93,6 +122,15 @@ static void DrawThickArc(FPrimitiveDrawInterface* PDI, const FVector& InCenter,
 }
 
 
+UGizmoElementArc::UGizmoElementArc()
+{
+	LineThickness = 0.0f;
+	HoverLineThicknessMultiplier = 1.0f;
+	InteractLineThicknessMultiplier = 1.0f;
+	SelectLineThicknessMultiplier = 1.0f;
+	SubdueLineThicknessMultiplier = 1.0f;
+}
+
 void UGizmoElementArc::Render(IToolsContextRenderAPI* RenderAPI, const FRenderTraversalState& RenderState)
 {
 	FRenderTraversalState CurrentRenderState(RenderState);
@@ -113,7 +151,9 @@ void UGizmoElementArc::Render(IToolsContextRenderAPI* RenderAPI, const FRenderTr
 	const FVector WorldAxis1 = CurrentRenderState.LocalToWorldTransform.TransformVectorNoScale(Axis1);
 	const FVector WorldNormal = WorldAxis0 ^ WorldAxis1;
 
-	bool bPartial = IsPartial(RenderAPI->GetSceneView(), WorldCenter, WorldNormal);
+	const FSceneView* View = RenderAPI->GetSceneView();
+
+	bool bPartial = IsPartial(View, WorldCenter, WorldNormal);
 	double Start, End;
 	if (bPartial)
 	{
@@ -134,14 +174,15 @@ void UGizmoElementArc::Render(IToolsContextRenderAPI* RenderAPI, const FRenderTr
 	const double WorldOuterRadius = Radius * CurrentRenderState.LocalToWorldTransform.GetScale3D().X;
 	const double WorldInnerRadius = InnerRadius * CurrentRenderState.LocalToWorldTransform.GetScale3D().X;
 	const FColor VertexColor = CurrentRenderState.GetCurrentVertexColor().ToFColor(false);
+	const FColor LineColor = CurrentRenderState.GetCurrentLineColor().ToFColor(false);
+	const double CurrentLineThickness = GetCurrentLineThickness(View->IsPerspectiveProjection(), View->FOV);
 
 	FPrimitiveDrawInterface* PDI = RenderAPI->GetPrimitiveDrawInterface();
 
 	DrawThickArc(PDI, WorldCenter, WorldAxis0, WorldAxis1, WorldOuterRadius, WorldInnerRadius,
-		Start, End, NumSegments, UseMaterial->GetRenderProxy(), VertexColor);
+		Start, End, NumSegments, UseMaterial->GetRenderProxy(), VertexColor, LineColor, CurrentLineThickness);
 }
 
-
 FInputRayHit UGizmoElementArc::LineTrace(const UGizmoViewContext* ViewContext, const FLineTraceTraversalState& LineTraceState, const FVector& RayOrigin, const FVector& RayDirection)
 {
 	FLineTraceTraversalState CurrentLineTraceState(LineTraceState);
@@ -227,4 +268,3 @@ double UGizmoElementArc::GetInnerRadius() const
 {
 	return InnerRadius;
 }
-
diff --git a/Engine/Source/Runtime/InteractiveToolsFramework/Public/BaseGizmos/GizmoElementArc.h b/Engine/Source/Runtime/InteractiveToolsFramework/Public/BaseGizmos/GizmoElementArc.h
--- a/Engine/Source/Runtime/InteractiveToolsFramework/Public/BaseGizmos/GizmoElementArc.h
+++ b/Engine/Source/Runtime/InteractiveToolsFramework/Public/BaseGizmos/GizmoElementArc.h
@@ -22,6 +22,8 @@ class UGizmoElementArc : public UGizmoElementCircleBase
 
 public:
 
+	UGizmoElementArc();
+
 	//~ Begin UGizmoElementBase Interface.
 	INTERACTIVETOOLSFRAMEWORK_API virtual void Render(IToolsContextRenderAPI* RenderAPI, const FRenderTraversalState& RenderState) override;
 	INTERACTIVETOOLSFRAMEWORK_API virtual FInputRayHit LineTrace(const UGizmoViewContext* ViewContext, const FLineTraceTraversalState& LineTraceState, const FVector& RayOrigin, const FVector& RayDirection) override;
