[ITF] Added overloads to Gizmo Elements to support same part id, multiple elements

diff --git a/Engine/Source/Runtime/InteractiveToolsFramework/Private/BaseGizmos/GizmoElementBase.cpp b/Engine/Source/Runtime/InteractiveToolsFramework/Private/BaseGizmos/GizmoElementBase.cpp
--- a/Engine/Source/Runtime/InteractiveToolsFramework/Private/BaseGizmos/GizmoElementBase.cpp
+++ b/Engine/Source/Runtime/InteractiveToolsFramework/Private/BaseGizmos/GizmoElementBase.cpp
@@ -423,6 +423,20 @@ uint32 UGizmoElementBase::GetPartIdentifier()
 	return PartIdentifier;
 }
 
+uint32 UGizmoElementBase::GetPartIdentifier() const
+{
+	return PartIdentifier;
+}
+
+void UGizmoElementBase::SetPartIdentifier(uint32 InPartId, const bool bInOverrideUnsetChildren, const bool bInOverrideSet)
+{
+	if (bInOverrideSet || GetPartIdentifier() == 0)
+	{
+		// The base implementation has no children, so we just call the default Setter
+		SetPartIdentifier(InPartId);	
+	}
+}
+
 void UGizmoElementBase::SetElementState(EGizmoElementState InElementState)
 {
 	ElementState = InElementState;
@@ -443,7 +457,13 @@ EGizmoElementInteractionState UGizmoElementBase::GetElementInteractionState() co
 	return ElementInteractionState;
 }
 
+
 bool UGizmoElementBase::UpdatePartHittableState(bool bHittable, uint32 InPartIdentifier)
+{
+	return UpdatePartHittableState(bHittable, InPartIdentifier, false);
+}
+
+bool UGizmoElementBase::UpdatePartHittableState(bool bHittable, uint32 InPartIdentifier, bool bInAllowMultipleElements)
 {
 	if (InPartIdentifier != PartIdentifier)
 	{
@@ -460,6 +480,11 @@ TOptional<bool> UGizmoElementBase::GetPartHittableState(uint32 InPartIdentifier)
 }
 
 bool UGizmoElementBase::UpdatePartVisibleState(bool bVisible, uint32 InPartIdentifier)
+{
+	return UpdatePartVisibleState(bVisible, InPartIdentifier, false);
+}
+
+bool UGizmoElementBase::UpdatePartVisibleState(bool bVisible, uint32 InPartIdentifier, bool bInAllowMultipleElements)
 {
 	if (InPartIdentifier != PartIdentifier)
 	{
@@ -476,6 +501,11 @@ TOptional<bool> UGizmoElementBase::GetPartVisibleState(uint32 InPartIdentifier)
 }
 
 bool UGizmoElementBase::UpdatePartInteractionState(EGizmoElementInteractionState InInteractionState, uint32 InPartIdentifier)
+{
+	return UpdatePartInteractionState(InInteractionState, InPartIdentifier, false);
+}
+
+bool UGizmoElementBase::UpdatePartInteractionState(EGizmoElementInteractionState InInteractionState, uint32 InPartIdentifier, bool bInAllowMultipleElements)
 {
 	if (InPartIdentifier != PartIdentifier)
 	{
@@ -691,6 +721,7 @@ FLinearColor UGizmoElementBase::GetVertexColor() const
 {
 	return MeshRenderAttributes.VertexColor.GetColor();
 }
+
 bool UGizmoElementBase::HasVertexColor() const
 {
 	return MeshRenderAttributes.VertexColor.bHasValue;
@@ -715,6 +746,7 @@ FLinearColor UGizmoElementBase::GetHoverVertexColor() const
 {
 	return MeshRenderAttributes.HoverVertexColor.GetColor();
 }
+
 bool UGizmoElementBase::HasHoverVertexColor() const
 {
 	return MeshRenderAttributes.HoverVertexColor.bHasValue;
@@ -739,6 +771,7 @@ FLinearColor UGizmoElementBase::GetInteractVertexColor() const
 {
 	return MeshRenderAttributes.InteractVertexColor.GetColor();
 }
+
 bool UGizmoElementBase::HasInteractVertexColor() const
 {
 	return MeshRenderAttributes.InteractVertexColor.bHasValue;
diff --git a/Engine/Source/Runtime/InteractiveToolsFramework/Private/BaseGizmos/GizmoElementGroup.cpp b/Engine/Source/Runtime/InteractiveToolsFramework/Private/BaseGizmos/GizmoElementGroup.cpp
--- a/Engine/Source/Runtime/InteractiveToolsFramework/Private/BaseGizmos/GizmoElementGroup.cpp
+++ b/Engine/Source/Runtime/InteractiveToolsFramework/Private/BaseGizmos/GizmoElementGroup.cpp
@@ -26,13 +26,10 @@ void UGizmoElementGroupBase::Render(IToolsContextRenderAPI* RenderAPI, const FRe
 
 		// Continue render even if not visible so all transforms will be cached 
 		// for subsequent line tracing.
-		for (UGizmoElementBase* Element : Elements)
+		ForEachElement([RenderAPI, &CurrentRenderState](UGizmoElementBase* Element)
 		{
-			if (Element)
-			{
-				Element->Render(RenderAPI, CurrentRenderState);
-			}
-		}
+			Element->Render(RenderAPI, CurrentRenderState);
+		});
 	}
 }
 
@@ -47,13 +44,10 @@ void UGizmoElementGroupBase::DrawHUD(FCanvas* Canvas, IToolsContextRenderAPI* Re
 
 		// Continue render even if not visible so all transforms will be cached
 		// for subsequent line tracing.
-		for (UGizmoElementBase* Element : Elements)
+		ForEachElement([Canvas, RenderAPI, &CurrentRenderState](UGizmoElementBase* Element)
 		{
-			if (Element)
-			{
-				Element->DrawHUD(Canvas, RenderAPI, CurrentRenderState);
-			}
-		}
+			Element->DrawHUD(Canvas, RenderAPI, CurrentRenderState);
+		});
 	}
 }
 
@@ -68,26 +62,44 @@ FInputRayHit UGizmoElementGroupBase::LineTrace(const UGizmoViewContext* ViewCont
 	{
 		ApplyUniformConstantScaleToTransform(CurrentLineTraceState.PixelToWorldScale, CurrentLineTraceState.LocalToWorldTransform);
 
-		for (UGizmoElementBase* Element : Elements)
+		ForEachElement([this, ViewContext, &CurrentLineTraceState, &RayOrigin, &RayDirection, &Hit](UGizmoElementBase* Element)
 		{
-			if (Element)
+			FInputRayHit NewHit = Element->LineTrace(ViewContext, CurrentLineTraceState, RayOrigin, RayDirection);
+			if (!Hit.bHit || NewHit.HitDepth < Hit.HitDepth)
 			{
-				FInputRayHit NewHit = Element->LineTrace(ViewContext, CurrentLineTraceState, RayOrigin, RayDirection);
-				if (!Hit.bHit || NewHit.HitDepth < Hit.HitDepth)
+				Hit = NewHit;
+				if (bHitOwner)
 				{
-					Hit = NewHit;
-					if (bHitOwner)
-					{
-						Hit.SetHitObject(this);
-						Hit.HitIdentifier = PartIdentifier;
-					}
+					Hit.SetHitObject(this);
+					Hit.HitIdentifier = PartIdentifier;
 				}
 			}
-		}
+		});
 	}
 	return Hit;
 }
 
+void UGizmoElementGroupBase::SetPartIdentifier(uint32 InPartId, const bool bInOverrideUnsetChildren, const bool bInOverrideSet)
+{
+	if (!bInOverrideUnsetChildren)
+	{
+		Super::SetPartIdentifier(InPartId, bInOverrideUnsetChildren);
+	}
+
+	// Only set if the part identifier is current Default (0), and bInOverrideSet is true 
+	if (bInOverrideSet || GetPartIdentifier() == 0)
+	{
+		Super::SetPartIdentifier(InPartId);
+	}
+
+	ForEachElement([InPartId, bInOverrideUnsetChildren](UGizmoElementBase* Element)
+	{
+		// By default, bOverrideChildIfSet is true, but we only want that for the root element
+		constexpr bool bOverrideChildIfSet = false;
+		Element->SetPartIdentifier(InPartId, bInOverrideUnsetChildren, bOverrideChildIfSet);
+	});
+}
+
 void UGizmoElementGroupBase::Add(UGizmoElementBase* InElement)
 {
 	if (!Elements.Contains(InElement))
@@ -108,22 +120,43 @@ void UGizmoElementGroupBase::Remove(UGizmoElementBase* InElement)
 	}
 }
 
-bool UGizmoElementGroupBase::UpdatePartVisibleState(bool bVisible, uint32 InPartIdentifier)
+void UGizmoElementGroupBase::ForEachElement(const TFunctionRef<void(UGizmoElementBase*)>& InFunc)
 {
-	if (Super::UpdatePartVisibleState(bVisible, InPartIdentifier))
+	for (UGizmoElementBase* Element : Elements)
 	{
-		return true;
+		if (Element)
+		{
+			InFunc(Element);
+		}
 	}
+}
 
-	for (UGizmoElementBase* Element : Elements)
+bool UGizmoElementGroupBase::UpdatePartVisibleState(bool bVisible, uint32 InPartIdentifier, bool bInAllowMultipleElements)
+{
+	bool bWasUpdated = false;
+
+	if (Super::UpdatePartVisibleState(bVisible, InPartIdentifier, bInAllowMultipleElements))
 	{
-		if (Element && Element->UpdatePartVisibleState(bVisible, InPartIdentifier))
+		bWasUpdated = true;
+		if (bWasUpdated && !bInAllowMultipleElements)
 		{
 			return true;
 		}
 	}
 
-	return false;
+	for (UGizmoElementBase* Element : Elements)
+	{
+		if (Element)
+		{
+			bWasUpdated |= Element->UpdatePartVisibleState(bVisible, InPartIdentifier, bInAllowMultipleElements);
+			if (bWasUpdated && !bInAllowMultipleElements)
+			{
+				return true;
+			}
+		}
+	}
+
+	return bWasUpdated;
 }
 
 TOptional<bool> UGizmoElementGroupBase::GetPartVisibleState(uint32 InPartIdentifier) const
@@ -150,22 +183,32 @@ TOptional<bool> UGizmoElementGroupBase::GetPartVisibleState(uint32 InPartIdentif
 }
 
 
-bool UGizmoElementGroupBase::UpdatePartHittableState(bool bHittable, uint32 InPartIdentifier)
+bool UGizmoElementGroupBase::UpdatePartHittableState(bool bHittable, uint32 InPartIdentifier, bool bInAllowMultipleElements)
 {
-	if (Super::UpdatePartHittableState(bHittable, InPartIdentifier))
+	bool bWasUpdated = false;
+
+	if (Super::UpdatePartHittableState(bHittable, InPartIdentifier, bInAllowMultipleElements))
 	{
-		return true;
+		bWasUpdated = true;
+		if (bWasUpdated && !bInAllowMultipleElements)
+		{
+			return true;
+		}
 	}
 
 	for (UGizmoElementBase* Element : Elements)
 	{
-		if (Element && Element->UpdatePartHittableState(bHittable, InPartIdentifier))
+		if (Element)
 		{
-			return true;
+			bWasUpdated |= Element->UpdatePartHittableState(bHittable, InPartIdentifier, bInAllowMultipleElements);
+			if (bWasUpdated && !bInAllowMultipleElements)
+			{
+				return true;
+			}
 		}
 	}
 
-	return false;
+	return bWasUpdated;
 }
 
 TOptional<bool> UGizmoElementGroupBase::GetPartHittableState(uint32 InPartIdentifier) const
@@ -191,22 +234,32 @@ TOptional<bool> UGizmoElementGroupBase::GetPartHittableState(uint32 InPartIdenti
 	return TOptional<bool>();
 }
 
-bool UGizmoElementGroupBase::UpdatePartInteractionState(EGizmoElementInteractionState InInteractionState, uint32 InPartIdentifier)
+bool UGizmoElementGroupBase::UpdatePartInteractionState(EGizmoElementInteractionState InInteractionState, uint32 InPartIdentifier, bool bInAllowMultipleElements)
 {
-	if (Super::UpdatePartInteractionState(InInteractionState, InPartIdentifier))
+	bool bWasUpdated = false;
+
+	if (Super::UpdatePartInteractionState(InInteractionState, InPartIdentifier, bInAllowMultipleElements))
 	{
-		return true;
+		bWasUpdated = true;
+		if (bWasUpdated && !bInAllowMultipleElements)
+		{
+			return true;
+		}
 	}
 
 	for (UGizmoElementBase* Element : Elements)
 	{
-		if (Element && Element->UpdatePartInteractionState(InInteractionState, InPartIdentifier))
+		if (Element)
 		{
-			return true;
+			bWasUpdated |= Element->UpdatePartInteractionState(InInteractionState, InPartIdentifier, bInAllowMultipleElements);
+			if (bWasUpdated && !bInAllowMultipleElements)
+			{
+				return true;
+			}
 		}
 	}
 
-	return false;
+	return bWasUpdated;
 }
 
 
diff --git a/Engine/Source/Runtime/InteractiveToolsFramework/Private/BaseGizmos/GizmoElementHitTargets.cpp b/Engine/Source/Runtime/InteractiveToolsFramework/Private/BaseGizmos/GizmoElementHitTargets.cpp
--- a/Engine/Source/Runtime/InteractiveToolsFramework/Private/BaseGizmos/GizmoElementHitTargets.cpp
+++ b/Engine/Source/Runtime/InteractiveToolsFramework/Private/BaseGizmos/GizmoElementHitTargets.cpp
@@ -8,7 +8,6 @@
 
 #include UE_INLINE_GENERATED_CPP_BY_NAME(GizmoElementHitTargets)
 
-
 FInputRayHit UGizmoElementHitTarget::IsHit(const FInputDeviceRay& ClickPos) const
 {
 	if (GizmoElement && GizmoViewContext && GizmoTransformProxy && (!Condition || Condition(ClickPos)))
@@ -56,6 +55,25 @@ void UGizmoElementHitTarget::UpdateInteractingState(bool bInteracting)
 	}
 }
 
+void UGizmoElementHitTarget::UpdateSelectedState(bool bSelected)
+{
+	if (!GizmoElement)
+	{
+		return;
+	}
+
+	if (bSelected)
+	{
+		GizmoElement->SetElementInteractionState(EGizmoElementInteractionState::Selected);
+	}
+	// If this element is interacting, do not update the select state. This is necessary because when transitioning
+	// from selected to interacting, the gizmo framework updates interacting to true before updating selected to false.
+	else if (GizmoElement->GetElementInteractionState() != EGizmoElementInteractionState::Interacting)
+	{
+		GizmoElement->SetElementInteractionState(EGizmoElementInteractionState::None);
+	}
+}
+
 UGizmoElementHitTarget* UGizmoElementHitTarget::Construct(UGizmoElementBase* InGizmoElement, UGizmoViewContext* InGizmoViewContext, UObject* Outer)
 {
 	UGizmoElementHitTarget* NewHitTarget = NewObject<UGizmoElementHitTarget>(Outer);
@@ -75,7 +93,12 @@ FInputRayHit UGizmoElementHitMultiTarget::IsHit(const FInputDeviceRay& ClickPos)
 	return FInputRayHit();
 }
 
-void UGizmoElementHitMultiTarget::UpdateHoverState(bool bInHovering, uint32 PartIdentifier)
+void UGizmoElementHitMultiTarget::UpdateHoverState(bool bHovering, uint32 PartIdentifier)
+{
+	UpdateHoverState(bHovering, PartIdentifier, false);
+}
+
+void UGizmoElementHitMultiTarget::UpdateHoverState(bool bInHovering, uint32 PartIdentifier, bool bInAllowMultipleElements)
 {
 	if (!GizmoElement)
 	{
@@ -84,7 +107,7 @@ void UGizmoElementHitMultiTarget::UpdateHoverState(bool bInHovering, uint32 Part
 
 	if (bInHovering)
 	{
-		GizmoElement->UpdatePartInteractionState(EGizmoElementInteractionState::Hovering, PartIdentifier);
+		GizmoElement->UpdatePartInteractionState(EGizmoElementInteractionState::Hovering, PartIdentifier, bInAllowMultipleElements);
 	}
 	else
 	{
@@ -96,11 +119,16 @@ void UGizmoElementHitMultiTarget::UpdateHoverState(bool bInHovering, uint32 Part
 			return;
 		}
 
-		GizmoElement->UpdatePartInteractionState(EGizmoElementInteractionState::None, PartIdentifier);
+		GizmoElement->UpdatePartInteractionState(EGizmoElementInteractionState::None, PartIdentifier, bInAllowMultipleElements);
 	}
 }
 
-void UGizmoElementHitMultiTarget::UpdateInteractingState(bool bInInteracting, uint32 PartIdentifier)
+void UGizmoElementHitMultiTarget::UpdateInteractingState(bool bInteracting, uint32 PartIdentifier)
+{
+	UpdateInteractingState(bInteracting, PartIdentifier, false);
+}
+
+void UGizmoElementHitMultiTarget::UpdateInteractingState(bool bInInteracting, uint32 PartIdentifier, bool bInAllowMultipleElements)
 {
 	if (!GizmoElement)
 	{
@@ -109,22 +137,44 @@ void UGizmoElementHitMultiTarget::UpdateInteractingState(bool bInInteracting, ui
 
 	if (bInInteracting)
 	{
-		GizmoElement->UpdatePartInteractionState(EGizmoElementInteractionState::Interacting, PartIdentifier);
+		GizmoElement->UpdatePartInteractionState(EGizmoElementInteractionState::Interacting, PartIdentifier, bInAllowMultipleElements);
 	}
 	else
 	{
-		GizmoElement->UpdatePartInteractionState(EGizmoElementInteractionState::None, PartIdentifier);
+		GizmoElement->UpdatePartInteractionState(EGizmoElementInteractionState::None, PartIdentifier, bInAllowMultipleElements);
+	}
+}
+
+void UGizmoElementHitMultiTarget::UpdateSelectedState(bool bSelected, uint32 PartIdentifier, bool bInAllowMultipleElements)
+{
+	if (!GizmoElement)
+	{
+		return;
+	}
+
+	if (bSelected)
+	{
+		GizmoElement->UpdatePartInteractionState(EGizmoElementInteractionState::Selected, PartIdentifier, bInAllowMultipleElements);
+	}
+	else
+	{
+		GizmoElement->UpdatePartInteractionState(EGizmoElementInteractionState::None, PartIdentifier, bInAllowMultipleElements);
 	}
 }
 
 void UGizmoElementHitMultiTarget::UpdateHittableState(bool bHittable, uint32 PartIdentifier)
+{
+	UpdateHittableState(bHittable, PartIdentifier, false);
+}
+
+void UGizmoElementHitMultiTarget::UpdateHittableState(bool bHittable, uint32 PartIdentifier, bool bInAllowMultipleElements)
 {
 	if (!GizmoElement)
 	{
 		return;
 	}
 
-	GizmoElement->UpdatePartHittableState(bHittable, PartIdentifier);
+	GizmoElement->UpdatePartHittableState(bHittable, PartIdentifier, bInAllowMultipleElements);
 }
 
 UGizmoElementHitMultiTarget* UGizmoElementHitMultiTarget::Construct(UGizmoElementBase* InGizmoElement, UGizmoViewContext* InGizmoViewContext, UObject* Outer)
@@ -134,7 +184,3 @@ UGizmoElementHitMultiTarget* UGizmoElementHitMultiTarget::Construct(UGizmoElemen
 	NewHitTarget->GizmoViewContext = InGizmoViewContext;
 	return NewHitTarget;
 }
-
-
-
-
diff --git a/Engine/Source/Runtime/InteractiveToolsFramework/Public/BaseGizmos/GizmoElementBase.h b/Engine/Source/Runtime/InteractiveToolsFramework/Public/BaseGizmos/GizmoElementBase.h
--- a/Engine/Source/Runtime/InteractiveToolsFramework/Public/BaseGizmos/GizmoElementBase.h
+++ b/Engine/Source/Runtime/InteractiveToolsFramework/Public/BaseGizmos/GizmoElementBase.h
@@ -170,6 +170,11 @@ public:
 	// that do not correspond to any gizmo part, such as non-hittable decorative elements.
 	INTERACTIVETOOLSFRAMEWORK_API virtual void SetPartIdentifier(uint32 InPartId);
 	INTERACTIVETOOLSFRAMEWORK_API virtual uint32 GetPartIdentifier();
+	INTERACTIVETOOLSFRAMEWORK_API virtual uint32 GetPartIdentifier() const;
+
+	// Sets the part identifier for this element and all of its children whose current ID is 0/Default.
+	// If bInOverrideSet is true, the PartId will be set even if it's currently 0. 
+	INTERACTIVETOOLSFRAMEWORK_API virtual void SetPartIdentifier(uint32 InPartId, const bool bInOverrideUnsetChildren, const bool bInOverrideSet = true);
 
 	// Object type bitmask indicating whether this object is visible or hittable or both
 	INTERACTIVETOOLSFRAMEWORK_API virtual void SetElementState(EGizmoElementState InElementState);
@@ -180,20 +185,35 @@ public:
 	INTERACTIVETOOLSFRAMEWORK_API virtual EGizmoElementInteractionState GetElementInteractionState() const;
 
 	// Update element's visibility state if element is associated with the specified gizmo part, return true if part was found.
+	UE_DEPRECATED(5.7, "Use UpdatePartVisibleState with bInAllowMultipleElements instead.")
 	INTERACTIVETOOLSFRAMEWORK_API virtual bool UpdatePartVisibleState(bool bVisible, uint32 InPartIdentifier);
 
+	// Update element's visibility state if element is associated with the specified gizmo part, return true if part was found.
+	// Setting bAllowMultipleElements to true will allow multiple elements with the same Part Id to be updated.
+	INTERACTIVETOOLSFRAMEWORK_API virtual bool UpdatePartVisibleState(bool bVisible, uint32 InPartIdentifier, bool bInAllowMultipleElements);
+
 	// Get element's visible state for element associated with the specified gizmo part, if part id was found.
 	INTERACTIVETOOLSFRAMEWORK_API virtual TOptional<bool> GetPartVisibleState(uint32 InPartIdentifier) const;
 
 	// Update element's hittable state if element is associated with the specified gizmo part, return true if part id was found.
+	UE_DEPRECATED(5.7, "Use UpdatePartHittableState with bInAllowMultipleElements instead.")
 	INTERACTIVETOOLSFRAMEWORK_API virtual bool UpdatePartHittableState(bool bHittable, uint32 InPartIdentifier);
 
+	// Update element's hittable state if element is associated with the specified gizmo part, return true if part id was found.
+	// Setting bAllowMultipleElements to true will allow multiple elements with the same Part Id to be updated.
+	INTERACTIVETOOLSFRAMEWORK_API virtual bool UpdatePartHittableState(bool bHittable, uint32 InPartIdentifier, bool bInAllowMultipleElements);
+
 	// Get element's hittable state for element associated with the specified gizmo part, if part id was found.
 	INTERACTIVETOOLSFRAMEWORK_API virtual TOptional<bool> GetPartHittableState(uint32 InPartIdentifier) const;
 
 	// Update element's interaction state if element is associated with the specified gizmo part, return true if part id was found.
+	UE_DEPRECATED(5.7, "Use UpdatePartInteractionState with bInAllowMultipleElements instead.")
 	INTERACTIVETOOLSFRAMEWORK_API virtual bool UpdatePartInteractionState(EGizmoElementInteractionState InInteractionState, uint32 InPartIdentifier);
 
+	// Update element's interaction state if element is associated with the specified gizmo part, return true if part id was found.
+	// Setting bAllowMultipleElements to true will allow multiple elements with the same Part Id to be updated.
+	INTERACTIVETOOLSFRAMEWORK_API virtual bool UpdatePartInteractionState(EGizmoElementInteractionState InInteractionState, uint32 InPartIdentifier, bool bInAllowMultipleElements);
+
 	// Get element's interaction state for element associated with the specified gizmo part, if part id was found.
 	INTERACTIVETOOLSFRAMEWORK_API virtual TOptional<EGizmoElementInteractionState> GetPartInteractionState(uint32 InPartIdentifier) const;
 
diff --git a/Engine/Source/Runtime/InteractiveToolsFramework/Public/BaseGizmos/GizmoElementGroup.h b/Engine/Source/Runtime/InteractiveToolsFramework/Public/BaseGizmos/GizmoElementGroup.h
--- a/Engine/Source/Runtime/InteractiveToolsFramework/Public/BaseGizmos/GizmoElementGroup.h
+++ b/Engine/Source/Runtime/InteractiveToolsFramework/Public/BaseGizmos/GizmoElementGroup.h
@@ -17,31 +17,41 @@ class UGizmoElementGroupBase : public UGizmoElementLineBase
 {
 	GENERATED_BODY()
 
+	friend class FGizmoElementTraverser;
+
 public:
 
 	//~ Begin UGizmoElementBase Interface.
 	INTERACTIVETOOLSFRAMEWORK_API virtual void Render(IToolsContextRenderAPI* RenderAPI, const FRenderTraversalState& RenderState) override;
 	INTERACTIVETOOLSFRAMEWORK_API virtual void DrawHUD(FCanvas* Canvas, IToolsContextRenderAPI* RenderAPI, const FRenderTraversalState& RenderState) override;
 	INTERACTIVETOOLSFRAMEWORK_API virtual FInputRayHit LineTrace(const UGizmoViewContext* ViewContext, const FLineTraceTraversalState& LineTraceState, const FVector& RayOrigin, const FVector& RayDirection) override;
-	//~ End UGizmoElementBase Interface.
+
+	// Sets the part identifier for this element and all of its children whose current ID is 0/Default.
+	using UGizmoElementLineBase::SetPartIdentifier;
+
+	// Sets the part identifier for this element and all of its children whose current ID is 0/Default.
+	INTERACTIVETOOLSFRAMEWORK_API virtual void SetPartIdentifier(uint32 InPartId, const bool bInOverrideUnsetChildren, const bool bInOverrideSet = true) override;
 
 	// Update group and contained elements' visibility state for elements in specified gizmo parts, return true if part was found.
-	INTERACTIVETOOLSFRAMEWORK_API virtual bool UpdatePartVisibleState(bool bVisible, uint32 InPartIdentifier);
+	// Setting bAllowMultipleElements to true will allow multiple elements with the same Part Id to be updated.
+	INTERACTIVETOOLSFRAMEWORK_API virtual bool UpdatePartVisibleState(bool bVisible, uint32 InPartIdentifier, bool bInAllowMultipleElements = false) override;
 
 	// Get element's visible state for element associated with the specified gizmo part, if part was found.
-	INTERACTIVETOOLSFRAMEWORK_API virtual TOptional<bool> GetPartVisibleState(uint32 InPartIdentifier) const;
+	INTERACTIVETOOLSFRAMEWORK_API virtual TOptional<bool> GetPartVisibleState(uint32 InPartIdentifier) const override;
 
 	// Update group and contained elements' hittable state for elements in specified gizmo parts, return true if part was found.
-	INTERACTIVETOOLSFRAMEWORK_API virtual bool UpdatePartHittableState(bool bHittable, uint32 InPartIdentifier);
+	// Setting bAllowMultipleElements to true will allow multiple elements with the same Part Id to be updated.
+	INTERACTIVETOOLSFRAMEWORK_API virtual bool UpdatePartHittableState(bool bHittable, uint32 InPartIdentifier, bool bInAllowMultipleElements = false) override;
 
 	// Get element's hittable state for element associated with the specified gizmo part, if part was found.
-	INTERACTIVETOOLSFRAMEWORK_API virtual TOptional<bool> GetPartHittableState(uint32 InPartIdentifier) const;
+	INTERACTIVETOOLSFRAMEWORK_API virtual TOptional<bool> GetPartHittableState(uint32 InPartIdentifier) const override;
 
 	// Update group and contained elements' interaction state for elements in specified gizmo parts, return true if part was found.
-	INTERACTIVETOOLSFRAMEWORK_API virtual bool UpdatePartInteractionState(EGizmoElementInteractionState InInteractionState, uint32 InPartIdentifier);
+	// Setting bAllowMultipleElements to true will allow multiple elements with the same Part Id to be updated.
+	INTERACTIVETOOLSFRAMEWORK_API virtual bool UpdatePartInteractionState(EGizmoElementInteractionState InInteractionState, uint32 InPartIdentifier, bool bInAllowMultipleElements = false) override;
 
 	// Get element's interaction state for element associated with the specified gizmo part, if part was found.
-	INTERACTIVETOOLSFRAMEWORK_API virtual TOptional<EGizmoElementInteractionState> GetPartInteractionState(uint32 InPartIdentifiero) const;
+	INTERACTIVETOOLSFRAMEWORK_API virtual TOptional<EGizmoElementInteractionState> GetPartInteractionState(uint32 InPartIdentifier) const override;
 
 	// When true, maintains view-dependent constant scale for this gizmo object hierarchy
 	INTERACTIVETOOLSFRAMEWORK_API virtual void SetConstantScale(bool InConstantScale);
@@ -59,6 +69,9 @@ protected:
 	// Remove object from group, if it exists
 	INTERACTIVETOOLSFRAMEWORK_API virtual void Remove(UGizmoElementBase* InElement);
 
+	// Iterate over each element in the group
+	INTERACTIVETOOLSFRAMEWORK_API void ForEachElement(const TFunctionRef<void(UGizmoElementBase*)>& InFunc);
+
 protected:
 
 	// When true, maintains view-dependent constant scale for this gizmo object hierarchy
diff --git a/Engine/Source/Runtime/InteractiveToolsFramework/Public/BaseGizmos/GizmoElementHitTargets.h b/Engine/Source/Runtime/InteractiveToolsFramework/Public/BaseGizmos/GizmoElementHitTargets.h
--- a/Engine/Source/Runtime/InteractiveToolsFramework/Public/BaseGizmos/GizmoElementHitTargets.h
+++ b/Engine/Source/Runtime/InteractiveToolsFramework/Public/BaseGizmos/GizmoElementHitTargets.h
@@ -50,6 +50,8 @@ public:
 
 	INTERACTIVETOOLSFRAMEWORK_API virtual void UpdateInteractingState(bool bInteracting);
 
+	INTERACTIVETOOLSFRAMEWORK_API virtual void UpdateSelectedState(bool bSelected);
+
 public:
 	static INTERACTIVETOOLSFRAMEWORK_API UGizmoElementHitTarget* Construct(
 		UGizmoElementBase* InGizmoElement,
@@ -94,13 +96,21 @@ public:
 	TFunction<bool(const FInputDeviceRay&)> Condition = nullptr;
 
 public:
-	INTERACTIVETOOLSFRAMEWORK_API virtual FInputRayHit IsHit(const FInputDeviceRay& ClickPos) const;
+	INTERACTIVETOOLSFRAMEWORK_API virtual FInputRayHit IsHit(const FInputDeviceRay& ClickPos) const override;
+
+	INTERACTIVETOOLSFRAMEWORK_API virtual void UpdateHoverState(bool bHovering, uint32 PartIdentifier) override;
+
+	INTERACTIVETOOLSFRAMEWORK_API virtual void UpdateHoverState(bool bHovering, uint32 PartIdentifier, bool bInAllowMultipleElements);
 
-	INTERACTIVETOOLSFRAMEWORK_API virtual void UpdateHoverState(bool bHovering, uint32 PartIdentifier);
+	INTERACTIVETOOLSFRAMEWORK_API virtual void UpdateInteractingState(bool bInteracting, uint32 PartIdentifier) override;
 
-	INTERACTIVETOOLSFRAMEWORK_API virtual void UpdateInteractingState(bool bInteracting, uint32 PartIdentifier);
+	INTERACTIVETOOLSFRAMEWORK_API virtual void UpdateInteractingState(bool bInteracting, uint32 PartIdentifier, bool bInAllowMultipleElements);
 
-	INTERACTIVETOOLSFRAMEWORK_API virtual void UpdateHittableState(bool bHittable, uint32 PartIdentifier);
+	INTERACTIVETOOLSFRAMEWORK_API virtual void UpdateSelectedState(bool bSelected, uint32 PartIdentifier, bool bInAllowMultipleElements);
+
+	INTERACTIVETOOLSFRAMEWORK_API virtual void UpdateHittableState(bool bHittable, uint32 PartIdentifier) override;
+
+	INTERACTIVETOOLSFRAMEWORK_API virtual void UpdateHittableState(bool bHittable, uint32 PartIdentifier, bool bInAllowMultipleElements);
 
 public:
 	static INTERACTIVETOOLSFRAMEWORK_API UGizmoElementHitMultiTarget* Construct(
@@ -109,5 +119,3 @@ public:
 		UObject* Outer = (UObject*)GetTransientPackage());
 
 };
-
-
