[ITF] Added calls to Hit Targets to support new Selected and Subdued interaction states

diff --git a/Engine/Source/Runtime/InteractiveToolsFramework/Private/BaseGizmos/GizmoElementHitTargets.cpp b/Engine/Source/Runtime/InteractiveToolsFramework/Private/BaseGizmos/GizmoElementHitTargets.cpp
--- a/Engine/Source/Runtime/InteractiveToolsFramework/Private/BaseGizmos/GizmoElementHitTargets.cpp
+++ b/Engine/Source/Runtime/InteractiveToolsFramework/Private/BaseGizmos/GizmoElementHitTargets.cpp
@@ -8,6 +8,36 @@
 
 #include UE_INLINE_GENERATED_CPP_BY_NAME(GizmoElementHitTargets)
 
+namespace UE::InteractiveToolsFramework::Private
+{
+	namespace GizmoElementHitTargetLocals
+	{
+		constexpr bool bAllowMultiElementParts = false;
+		
+		// Toggles between the provided state and "None" state for the given gizmo element.
+		void ToggleElementInteractionState(UGizmoElementBase* InGizmoElement, const EGizmoElementInteractionState InState, const bool bInStateEnable)
+		{
+			if (!InGizmoElement)
+			{
+				return;
+			}
+
+			InGizmoElement->SetElementInteractionState(bInStateEnable ? InState : EGizmoElementInteractionState::None);
+		}
+
+		// Toggles between the provided state and "None" state for the given gizmo element.
+		void ToggleElementPartInteractionState(UGizmoElementBase* InGizmoElement, const uint32 InPartId, const EGizmoElementInteractionState InState, const bool bInStateEnable, const bool bInAllowMultipleElements = false)
+		{
+			if (!InGizmoElement)
+			{
+				return;
+			}
+
+			InGizmoElement->UpdatePartInteractionState(bInStateEnable ? InState : EGizmoElementInteractionState::None, InPartId, bInAllowMultipleElements);
+		}
+	}
+}
+
 FInputRayHit UGizmoElementHitTarget::IsHit(const FInputDeviceRay& ClickPos) const
 {
 	if (GizmoElement && GizmoViewContext && GizmoTransformProxy && (!Condition || Condition(ClickPos)))
@@ -40,38 +70,23 @@ void UGizmoElementHitTarget::UpdateHoverState(bool bHovering)
 
 void UGizmoElementHitTarget::UpdateInteractingState(bool bInteracting)
 {
-	if (!GizmoElement)
-	{
-		return;
-	}
+	using namespace UE::InteractiveToolsFramework::Private::GizmoElementHitTargetLocals;
 
-	if (bInteracting)
-	{
-		GizmoElement->SetElementInteractionState(EGizmoElementInteractionState::Interacting);
-	}
-	else
-	{
-		GizmoElement->SetElementInteractionState(EGizmoElementInteractionState::None);
-	}
+	ToggleElementInteractionState(GizmoElement, EGizmoElementInteractionState::Interacting, bInteracting);
 }
 
-void UGizmoElementHitTarget::UpdateSelectedState(bool bSelected)
+void UGizmoElementHitTarget::UpdateSelectedState(bool bInSelected)
 {
-	if (!GizmoElement)
-	{
-		return;
-	}
+	using namespace UE::InteractiveToolsFramework::Private::GizmoElementHitTargetLocals;
 
-	if (bSelected)
-	{
-		GizmoElement->SetElementInteractionState(EGizmoElementInteractionState::Selected);
-	}
-	// If this element is interacting, do not update the select state. This is necessary because when transitioning
-	// from selected to interacting, the gizmo framework updates interacting to true before updating selected to false.
-	else if (GizmoElement->GetElementInteractionState() != EGizmoElementInteractionState::Interacting)
-	{
-		GizmoElement->SetElementInteractionState(EGizmoElementInteractionState::None);
-	}
+	ToggleElementInteractionState(GizmoElement, EGizmoElementInteractionState::Selected, bInSelected);
+}
+
+void UGizmoElementHitTarget::UpdateSubdueState(bool bSubdued)
+{
+	using namespace UE::InteractiveToolsFramework::Private::GizmoElementHitTargetLocals;
+
+	ToggleElementInteractionState(GizmoElement, EGizmoElementInteractionState::Subdued, bSubdued);
 }
 
 UGizmoElementHitTarget* UGizmoElementHitTarget::Construct(UGizmoElementBase* InGizmoElement, UGizmoViewContext* InGizmoViewContext, UObject* Outer)
@@ -93,9 +108,11 @@ FInputRayHit UGizmoElementHitMultiTarget::IsHit(const FInputDeviceRay& ClickPos)
 	return FInputRayHit();
 }
 
-void UGizmoElementHitMultiTarget::UpdateHoverState(bool bHovering, uint32 PartIdentifier)
+void UGizmoElementHitMultiTarget::UpdateHoverState(bool bInHovering, uint32 PartIdentifier)
 {
-	UpdateHoverState(bHovering, PartIdentifier, false);
+	using namespace UE::InteractiveToolsFramework::Private::GizmoElementHitTargetLocals;
+	
+	return UpdateHoverState(bInHovering, PartIdentifier, bAllowMultiElementParts);
 }
 
 void UGizmoElementHitMultiTarget::UpdateHoverState(bool bInHovering, uint32 PartIdentifier, bool bInAllowMultipleElements)
@@ -123,48 +140,53 @@ void UGizmoElementHitMultiTarget::UpdateHoverState(bool bInHovering, uint32 Part
 	}
 }
 
-void UGizmoElementHitMultiTarget::UpdateInteractingState(bool bInteracting, uint32 PartIdentifier)
+void UGizmoElementHitMultiTarget::UpdateInteractingState(bool bInInteracting, uint32 PartIdentifier)
 {
-	UpdateInteractingState(bInteracting, PartIdentifier, false);
+	using namespace UE::InteractiveToolsFramework::Private::GizmoElementHitTargetLocals;
+
+	UpdateInteractingState(bInInteracting, PartIdentifier, bAllowMultiElementParts);
 }
 
 void UGizmoElementHitMultiTarget::UpdateInteractingState(bool bInInteracting, uint32 PartIdentifier, bool bInAllowMultipleElements)
 {
-	if (!GizmoElement)
-	{
-		return;
-	}
+	using namespace UE::InteractiveToolsFramework::Private::GizmoElementHitTargetLocals;
 
-	if (bInInteracting)
-	{
-		GizmoElement->UpdatePartInteractionState(EGizmoElementInteractionState::Interacting, PartIdentifier, bInAllowMultipleElements);
-	}
-	else
-	{
-		GizmoElement->UpdatePartInteractionState(EGizmoElementInteractionState::None, PartIdentifier, bInAllowMultipleElements);
-	}
+	ToggleElementPartInteractionState(GizmoElement, PartIdentifier, EGizmoElementInteractionState::Interacting, bInInteracting, bInAllowMultipleElements);
 }
 
-void UGizmoElementHitMultiTarget::UpdateSelectedState(bool bSelected, uint32 PartIdentifier, bool bInAllowMultipleElements)
+void UGizmoElementHitMultiTarget::UpdateSelectedState(bool bInSelected, uint32 PartIdentifier)
 {
-	if (!GizmoElement)
-	{
-		return;
-	}
+	using namespace UE::InteractiveToolsFramework::Private::GizmoElementHitTargetLocals;
 
-	if (bSelected)
-	{
-		GizmoElement->UpdatePartInteractionState(EGizmoElementInteractionState::Selected, PartIdentifier, bInAllowMultipleElements);
-	}
-	else
-	{
-		GizmoElement->UpdatePartInteractionState(EGizmoElementInteractionState::None, PartIdentifier, bInAllowMultipleElements);
-	}
+	UpdateSelectedState(bInSelected, PartIdentifier, bAllowMultiElementParts);
+}
+
+void UGizmoElementHitMultiTarget::UpdateSelectedState(bool bInSelected, uint32 PartIdentifier, bool bInAllowMultipleElements)
+{
+	using namespace UE::InteractiveToolsFramework::Private::GizmoElementHitTargetLocals;
+
+	ToggleElementPartInteractionState(GizmoElement, PartIdentifier, EGizmoElementInteractionState::Selected, bInSelected, bInAllowMultipleElements);
+}
+
+void UGizmoElementHitMultiTarget::UpdateSubdueState(bool bInSubdued, uint32 PartIdentifier)
+{
+	using namespace UE::InteractiveToolsFramework::Private::GizmoElementHitTargetLocals;
+
+	UpdateSubdueState(bInSubdued, PartIdentifier, bAllowMultiElementParts);
+}
+
+void UGizmoElementHitMultiTarget::UpdateSubdueState(bool bInSubdued, uint32 PartIdentifier, bool bInAllowMultipleElements)
+{
+	using namespace UE::InteractiveToolsFramework::Private::GizmoElementHitTargetLocals;
+	
+	ToggleElementPartInteractionState(GizmoElement, PartIdentifier, EGizmoElementInteractionState::Subdued, bInSubdued, bInAllowMultipleElements);
 }
 
 void UGizmoElementHitMultiTarget::UpdateHittableState(bool bHittable, uint32 PartIdentifier)
 {
-	UpdateHittableState(bHittable, PartIdentifier, false);
+	using namespace UE::InteractiveToolsFramework::Private::GizmoElementHitTargetLocals;
+
+	UpdateHittableState(bHittable, PartIdentifier, bAllowMultiElementParts);
 }
 
 void UGizmoElementHitMultiTarget::UpdateHittableState(bool bHittable, uint32 PartIdentifier, bool bInAllowMultipleElements)
diff --git a/Engine/Source/Runtime/InteractiveToolsFramework/Private/BaseGizmos/HitTargets.cpp b/Engine/Source/Runtime/InteractiveToolsFramework/Private/BaseGizmos/HitTargets.cpp
--- a/Engine/Source/Runtime/InteractiveToolsFramework/Private/BaseGizmos/HitTargets.cpp
+++ b/Engine/Source/Runtime/InteractiveToolsFramework/Private/BaseGizmos/HitTargets.cpp
@@ -23,11 +23,27 @@ void UGizmoLambdaHitTarget::UpdateHoverState(bool bHovering)
 	}
 }
 
-void UGizmoLambdaHitTarget::UpdateInteractingState(bool bHovering)
+void UGizmoLambdaHitTarget::UpdateInteractingState(bool bInteracting)
 {
 	if (UpdateInteractingFunction)
 	{
-		UpdateInteractingFunction(bHovering);
+		UpdateInteractingFunction(bInteracting);
+	}
+}
+
+void UGizmoLambdaHitTarget::UpdateSelectedState(bool bSelected)
+{
+	if (UpdateSelectedFunction)
+	{
+		UpdateSelectedFunction(bSelected);
+	}
+}
+
+void UGizmoLambdaHitTarget::UpdateSubdueState(bool bSubdued)
+{
+	if (UpdateSubdueFunction)
+	{
+		UpdateSubdueFunction(bSubdued);
 	}
 }
 
@@ -71,6 +87,22 @@ void UGizmoComponentHitTarget::UpdateInteractingState(bool bHovering)
 	}
 }
 
+void UGizmoComponentHitTarget::UpdateSelectedState(bool bSelected)
+{
+	if (UpdateSelectedFunction)
+	{
+		UpdateSelectedFunction(bSelected);
+	}
+}
+
+void UGizmoComponentHitTarget::UpdateSubdueState(bool bSubdued)
+{
+	if (UpdateSubduedFunction)
+	{
+		UpdateSubduedFunction(bSubdued);
+	}
+}
+
 UGizmoComponentHitTarget* UGizmoComponentHitTarget::Construct(UPrimitiveComponent* Component, UObject* Outer)
 {
 	UGizmoComponentHitTarget* NewHitTarget = NewObject<UGizmoComponentHitTarget>(Outer);
@@ -78,6 +110,3 @@ UGizmoComponentHitTarget* UGizmoComponentHitTarget::Construct(UPrimitiveComponen
 	return NewHitTarget;
 }
 
-
-
-
diff --git a/Engine/Source/Runtime/InteractiveToolsFramework/Public/BaseGizmos/GizmoElementHitTargets.h b/Engine/Source/Runtime/InteractiveToolsFramework/Public/BaseGizmos/GizmoElementHitTargets.h
--- a/Engine/Source/Runtime/InteractiveToolsFramework/Public/BaseGizmos/GizmoElementHitTargets.h
+++ b/Engine/Source/Runtime/InteractiveToolsFramework/Public/BaseGizmos/GizmoElementHitTargets.h
@@ -46,11 +46,13 @@ public:
 public:
 	INTERACTIVETOOLSFRAMEWORK_API virtual FInputRayHit IsHit(const FInputDeviceRay& ClickPos) const;
 
-	INTERACTIVETOOLSFRAMEWORK_API virtual void UpdateHoverState(bool bHovering);
+	INTERACTIVETOOLSFRAMEWORK_API virtual void UpdateHoverState(bool bHovering) override;
 
-	INTERACTIVETOOLSFRAMEWORK_API virtual void UpdateInteractingState(bool bInteracting);
+	INTERACTIVETOOLSFRAMEWORK_API virtual void UpdateInteractingState(bool bInteracting) override;
 
-	INTERACTIVETOOLSFRAMEWORK_API virtual void UpdateSelectedState(bool bSelected);
+	INTERACTIVETOOLSFRAMEWORK_API virtual void UpdateSelectedState(bool bInSelected) override;
+
+	INTERACTIVETOOLSFRAMEWORK_API virtual void UpdateSubdueState(bool bSubdued) override;
 
 public:
 	static INTERACTIVETOOLSFRAMEWORK_API UGizmoElementHitTarget* Construct(
@@ -106,7 +108,13 @@ public:
 
 	INTERACTIVETOOLSFRAMEWORK_API virtual void UpdateInteractingState(bool bInteracting, uint32 PartIdentifier, bool bInAllowMultipleElements);
 
-	INTERACTIVETOOLSFRAMEWORK_API virtual void UpdateSelectedState(bool bSelected, uint32 PartIdentifier, bool bInAllowMultipleElements);
+	INTERACTIVETOOLSFRAMEWORK_API virtual void UpdateSelectedState(bool bInSelected, uint32 PartIdentifier) override;
+
+	INTERACTIVETOOLSFRAMEWORK_API virtual void UpdateSelectedState(bool bInSelected, uint32 PartIdentifier, bool bInAllowMultipleElements);
+
+	INTERACTIVETOOLSFRAMEWORK_API virtual void UpdateSubdueState(bool bInSubdued, uint32 PartIdentifier) override;
+
+	INTERACTIVETOOLSFRAMEWORK_API virtual void UpdateSubdueState(bool bInSubdued, uint32 PartIdentifier, bool bInAllowMultipleElements);
 
 	INTERACTIVETOOLSFRAMEWORK_API virtual void UpdateHittableState(bool bHittable, uint32 PartIdentifier) override;
 
diff --git a/Engine/Source/Runtime/InteractiveToolsFramework/Public/BaseGizmos/GizmoInterfaces.h b/Engine/Source/Runtime/InteractiveToolsFramework/Public/BaseGizmos/GizmoInterfaces.h
--- a/Engine/Source/Runtime/InteractiveToolsFramework/Public/BaseGizmos/GizmoInterfaces.h
+++ b/Engine/Source/Runtime/InteractiveToolsFramework/Public/BaseGizmos/GizmoInterfaces.h
@@ -127,6 +127,20 @@ public:
 	 */
 	UFUNCTION()
 	virtual void UpdateInteractingState(bool bInteracting) = 0;
+
+	/*
+	 * Updates the selected state, indicating when the gizmo is currently selected,
+	 * typically upon the input device clicking the Standard gizmo. 
+	 */
+	UFUNCTION()
+	virtual void UpdateSelectedState(bool bSelected) = 0;
+
+	/*
+	 * Updates the subdued state, indicating when the element should have a subdued/background appearance, 
+	 * typically upon interacting with one of a set of gizmos representing 3 axis - the two not interacted with would be subdued. 
+	 */
+	UFUNCTION()
+	virtual void UpdateSubdueState(bool bSubdued) = 0;
 };
 
 UINTERFACE(MinimalAPI)
@@ -168,12 +182,25 @@ public:
 	UFUNCTION()
 	virtual void UpdateInteractingState(bool bInteracting, uint32 InPartIdentifier) = 0;
 
+	/*
+	 * Updates the selected state of the specified gizmo part, indicating when the gizmo part is currently selected,
+	 * typically upon the input device clicking the Standard gizmo.
+	 */
+	UFUNCTION()
+	virtual void UpdateSelectedState(bool bSelected, uint32 InPartIdentifier) = 0;
+
+	/*
+	 * Updates the subdued state of the specified gizmo part, indicating when the element should have a subdued/background appearance, 
+	 * typically upon interacting with one of a set of gizmos representing 3 axis - the two not interacted with would be subdued.
+	 */
+	UFUNCTION()
+	virtual void UpdateSubdueState(bool bSubdued, uint32 InPartIdentifier) = 0;
+
 	/*
 	 * Updates the hittable state of the specified gizmo part. 
 	 */
 	UFUNCTION()
 	virtual void UpdateHittableState(bool bHittable, uint32 InPartIdentifier) = 0;
-
 };
 
 
diff --git a/Engine/Source/Runtime/InteractiveToolsFramework/Public/BaseGizmos/HitTargets.h b/Engine/Source/Runtime/InteractiveToolsFramework/Public/BaseGizmos/HitTargets.h
--- a/Engine/Source/Runtime/InteractiveToolsFramework/Public/BaseGizmos/HitTargets.h
+++ b/Engine/Source/Runtime/InteractiveToolsFramework/Public/BaseGizmos/HitTargets.h
@@ -26,14 +26,23 @@ public:
 
 	/** This function is called to update interacting state of the target */
 	TFunction<void(bool)> UpdateInteractingFunction;
+	
+	/** This function is called to update selected state of the target */
+	TFunction<void(bool)> UpdateSelectedFunction;
 
+	/** This function is called to update subdued state of the target */
+	TFunction<void(bool)> UpdateSubdueFunction;
 
 public:
 	INTERACTIVETOOLSFRAMEWORK_API virtual FInputRayHit IsHit(const FInputDeviceRay& ClickPos) const;
 
-	INTERACTIVETOOLSFRAMEWORK_API virtual void UpdateHoverState(bool bHovering);
+	INTERACTIVETOOLSFRAMEWORK_API virtual void UpdateHoverState(bool bHovering) override;
 
-	INTERACTIVETOOLSFRAMEWORK_API virtual void UpdateInteractingState(bool bHovering);
+	INTERACTIVETOOLSFRAMEWORK_API virtual void UpdateInteractingState(bool bInteracting) override;
+
+	INTERACTIVETOOLSFRAMEWORK_API virtual void UpdateSelectedState(bool bSelected) override;
+
+	INTERACTIVETOOLSFRAMEWORK_API virtual void UpdateSubdueState(bool bSubdued) override;
 };
 
 
@@ -68,12 +77,22 @@ public:
 	/** This function is called to update interacting state of the target */
 	TFunction<void(bool)> UpdateInteractingFunction;
 
+	/** This function is called to update selected state of the target */
+	TFunction<void(bool)> UpdateSelectedFunction;
+
+	/** This function is called to update subdued state of the target */
+	TFunction<void(bool)> UpdateSubduedFunction;
+
 public:
 	INTERACTIVETOOLSFRAMEWORK_API virtual FInputRayHit IsHit(const FInputDeviceRay& ClickPos) const;
 
-	INTERACTIVETOOLSFRAMEWORK_API virtual void UpdateHoverState(bool bHovering);
+	INTERACTIVETOOLSFRAMEWORK_API virtual void UpdateHoverState(bool bHovering) override;
+
+	INTERACTIVETOOLSFRAMEWORK_API virtual void UpdateInteractingState(bool bHovering) override;
+
+	INTERACTIVETOOLSFRAMEWORK_API virtual void UpdateSelectedState(bool bSelected) override;
 
-	INTERACTIVETOOLSFRAMEWORK_API virtual void UpdateInteractingState(bool bHovering);
+	INTERACTIVETOOLSFRAMEWORK_API virtual void UpdateSubdueState(bool bSubdued) override;
 
 public:
 	static INTERACTIVETOOLSFRAMEWORK_API UGizmoComponentHitTarget* Construct(
