[ITF] Separated GizmoElementArrow and GizmoElementArrowHead

diff --git a/Engine/Source/Runtime/InteractiveToolsFramework/Private/BaseGizmos/GizmoElementArrow.cpp b/Engine/Source/Runtime/InteractiveToolsFramework/Private/BaseGizmos/GizmoElementArrow.cpp
--- a/Engine/Source/Runtime/InteractiveToolsFramework/Private/BaseGizmos/GizmoElementArrow.cpp
+++ b/Engine/Source/Runtime/InteractiveToolsFramework/Private/BaseGizmos/GizmoElementArrow.cpp
@@ -1,12 +1,15 @@
 // Copyright Epic Games, Inc. All Rights Reserved.
  
 #include "BaseGizmos/GizmoElementArrow.h"
+
 #include "BaseGizmos/GizmoElementBox.h"
 #include "BaseGizmos/GizmoElementCone.h"
 #include "BaseGizmos/GizmoElementCylinder.h"
+#include "BaseGizmos/GizmoElementSphere.h"
 #include "BaseGizmos/GizmoRenderingUtil.h"
 #include "InputState.h"
 #include "Materials/MaterialInterface.h"
+#include "PrimitiveDrawInterface.h"
 
 #include UE_INLINE_GENERATED_CPP_BY_NAME(GizmoElementArrow)
 
@@ -27,10 +30,16 @@ using ThresholdGuard = TGuardValue_Bitfield_Cleanup<TFunction<void()>>;
 
 UGizmoElementArrow::UGizmoElementArrow()
 {
+	bHitOwner = true;
+
 	HeadType = EGizmoElementArrowHeadType::Cone;
+
 	CylinderElement = NewObject<UGizmoElementCylinder>();
-	ConeElement = NewObject<UGizmoElementCone>();
-	BoxElement = nullptr;
+	UGizmoElementArrow::Add(CylinderElement);
+
+	HeadElement = NewObject<UGizmoElementArrowHead>();
+	HeadElement->SetType(HeadType);
+	UGizmoElementArrow::Add(HeadElement);
 }
 
 void UGizmoElementArrow::Render(IToolsContextRenderAPI* RenderAPI, const FRenderTraversalState& RenderState)
@@ -55,16 +64,8 @@ void UGizmoElementArrow::Render(IToolsContextRenderAPI* RenderAPI, const FRender
 		check(CylinderElement);
 		CylinderElement->Render(RenderAPI, CurrentRenderState);
 
-		if (HeadType == EGizmoElementArrowHeadType::Cone)
-		{
-			check(ConeElement);
-			ConeElement->Render(RenderAPI, CurrentRenderState);
-		}
-		else // (HeadType == EGizmoElementArrowHeadType::Cube)
-		{
-			check(BoxElement);
-			BoxElement->Render(RenderAPI, CurrentRenderState);
-		}
+		check(HeadElement);
+		HeadElement->Render(RenderAPI, CurrentRenderState);
 	}
 }
 
@@ -83,8 +84,7 @@ FInputRayHit UGizmoElementArrow::LineTrace(const UGizmoViewContext* ViewContext,
 	ThresholdGuard Guard([this]()
 	{
 		UpdatePixelThreshold(CylinderElement, PixelHitDistanceThreshold);
-		UpdatePixelThreshold(ConeElement, PixelHitDistanceThreshold);
-		UpdatePixelThreshold(BoxElement, PixelHitDistanceThreshold);
+		UpdatePixelThreshold(HeadElement, PixelHitDistanceThreshold);
 	});
 
 	// update threshold if hitting the mask
@@ -95,8 +95,7 @@ FInputRayHit UGizmoElementArrow::LineTrace(const UGizmoViewContext* ViewContext,
 		{
 			static constexpr float NoPixelHitThreshold = 0.f;
 			UpdatePixelThreshold(CylinderElement, NoPixelHitThreshold);
-			UpdatePixelThreshold(ConeElement, NoPixelHitThreshold);
-			UpdatePixelThreshold(BoxElement, NoPixelHitThreshold);
+			UpdatePixelThreshold(HeadElement, NoPixelHitThreshold);
 		}
 	}
 	
@@ -105,16 +104,8 @@ FInputRayHit UGizmoElementArrow::LineTrace(const UGizmoViewContext* ViewContext,
 
 	if (!Hit.bHit)
 	{
-		if (HeadType == EGizmoElementArrowHeadType::Cone)
-		{
-			check(ConeElement);
-			Hit = ConeElement->LineTrace(ViewContext, CurrentLineTraceState, RayOrigin, RayDirection);
-		}
-		else // (HeadType == EGizmoElementArrowHeadType::Cube)
-		{
-			check(BoxElement);
-			Hit = BoxElement->LineTrace(ViewContext, CurrentLineTraceState, RayOrigin, RayDirection);
-		}
+		check(HeadElement);
+		Hit = HeadElement->LineTrace(ViewContext, CurrentLineTraceState, RayOrigin, RayDirection);
 	}
 
 	if (Hit.bHit)
@@ -273,17 +264,6 @@ void UGizmoElementArrow::SetHeadType(EGizmoElementArrowHeadType InHeadType)
 	if (InHeadType != HeadType)
 	{
 		HeadType = InHeadType;
-
-		if (HeadType == EGizmoElementArrowHeadType::Cone)
-		{
-			ConeElement = NewObject<UGizmoElementCone>();
-			BoxElement = nullptr;
-		}
-		else // (HeadType == EGizmoElementArrowHeadType::Cube)
-		{
-			BoxElement = NewObject<UGizmoElementBox>();
-			ConeElement = nullptr;
-		}
 		UpdateArrowHead();
 	}
 }
@@ -341,30 +321,16 @@ void UGizmoElementArrow::UpdateArrowBody()
 
 void UGizmoElementArrow::UpdateArrowHead()
 {
-	if (HeadType == EGizmoElementArrowHeadType::Cone)
-	{
-		check(ConeElement);
-		// head length is multiplied by 0.9f to prevent gap between body cylinder and head cone
-		ConeElement->SetOrigin(Direction * (BodyLength + HeadLength * 0.9f));
-		ConeElement->SetDirection(-Direction);
-		ConeElement->SetHeight(HeadLength);
-		ConeElement->SetRadius(HeadRadius);
-		ConeElement->SetNumSides(NumSides);
-		ConeElement->SetElementInteractionState(ElementInteractionState);
-		ConeElement->SetPixelHitDistanceThreshold(PixelHitDistanceThreshold);
-		ConeElement->SetEndCaps(bEndCaps);
-	}
-	else // (HeadType == EGizmoElementArrowHeadType::Cube)
-	{
-		check(BoxElement);
-		BoxElement->SetCenter(Direction * (BodyLength + HeadLength * 0.5f));
-		BoxElement->SetUpDirection(Direction);
-		BoxElement->SetSideDirection(SideDirection);
-		BoxElement->SetDimensions(FVector(HeadLength, HeadLength, HeadLength));
-		BoxElement->SetElementInteractionState(ElementInteractionState);
-		BoxElement->SetPixelHitDistanceThreshold(PixelHitDistanceThreshold);
-	}
+	check(HeadElement);
+	HeadElement->SetType(HeadType);
+	HeadElement->SetCenter(Direction * BodyLength);
+	HeadElement->SetDirection(Direction);
+	HeadElement->SetSideDirection(SideDirection);
+	HeadElement->SetLength(HeadLength);
+	HeadElement->SetRadius(HeadRadius);
+	HeadElement->SetNumSides(NumSides);
+	HeadElement->SetElementInteractionState(ElementInteractionState);
+	HeadElement->SetPixelHitDistanceThreshold(PixelHitDistanceThreshold);
 
 	bUpdateArrowHead = false;
 }
-
diff --git a/Engine/Source/Runtime/InteractiveToolsFramework/Private/BaseGizmos/GizmoElementArrowHead.cpp b/Engine/Source/Runtime/InteractiveToolsFramework/Private/BaseGizmos/GizmoElementArrowHead.cpp
new file mode 100644
--- /dev/null
+++ b/Engine/Source/Runtime/InteractiveToolsFramework/Private/BaseGizmos/GizmoElementArrowHead.cpp
@@ -0,0 +1,345 @@
+// Copyright Epic Games, Inc. All Rights Reserved.
+
+#include "BaseGizmos/GizmoElementArrowHead.h"
+
+#include "BaseGizmos/GizmoElementBox.h"
+#include "BaseGizmos/GizmoElementCone.h"
+#include "BaseGizmos/GizmoElementSphere.h"
+
+#include UE_INLINE_GENERATED_CPP_BY_NAME(GizmoElementArrowHead)
+
+namespace UE::InteractiveToolsFramework::Private
+{
+	namespace GizmoElementArrowHeadLocals
+	{
+		void UpdatePixelThreshold(UGizmoElementBase* InGizmoElement, const float InPixelHitThreshold)
+		{
+			if (InGizmoElement)
+			{
+				InGizmoElement->SetPixelHitDistanceThreshold(InPixelHitThreshold);
+			}
+		}
+
+		using FThresholdGuard = TGuardValue_Bitfield_Cleanup<TFunction<void()>>;
+	}
+}
+
+UGizmoElementArrowHead::UGizmoElementArrowHead()
+{
+	Type = EGizmoElementArrowHeadType::Cone;
+	bHitOwner = true;
+
+	ConeElement = NewObject<UGizmoElementCone>();
+	BoxElement = nullptr;
+	SphereElement = nullptr;
+
+	UGizmoElementArrowHead::Add(ConeElement);
+}
+
+void UGizmoElementArrowHead::Render(IToolsContextRenderAPI* RenderAPI, const FRenderTraversalState& RenderState)
+{
+	check(RenderAPI);
+
+	if (bUpdate)
+	{
+		Update();
+	}
+
+	FRenderTraversalState CurrentRenderState(RenderState);
+	const bool bVisibleViewDependent = UpdateRenderState(RenderAPI, Center, CurrentRenderState);
+
+	if (bVisibleViewDependent)
+	{
+		auto RenderElement = [RenderAPI, &CurrentRenderState](UGizmoElementBase* InElement)
+		{
+			check(InElement);
+			InElement->Render(RenderAPI, CurrentRenderState);
+		};
+
+		if (Type == EGizmoElementArrowHeadType::Cone)
+		{
+			RenderElement(ConeElement);
+		}
+		else if (Type == EGizmoElementArrowHeadType::Cube)
+		{
+			RenderElement(BoxElement);
+		}
+		else if (Type == EGizmoElementArrowHeadType::Sphere)
+		{
+			RenderElement(SphereElement);
+		}
+	}
+}
+
+FInputRayHit UGizmoElementArrowHead::LineTrace(const UGizmoViewContext* ViewContext, const FLineTraceTraversalState& LineTraceState, const FVector& RayOrigin, const FVector& RayDirection)
+{
+	using namespace UE::InteractiveToolsFramework::Private::GizmoElementArrowHeadLocals;
+
+	FLineTraceTraversalState CurrentLineTraceState(LineTraceState);
+	const bool bHittableViewDependent = UpdateLineTraceState(ViewContext, Center, CurrentLineTraceState);
+
+	if (!bHittableViewDependent)
+	{
+		return FInputRayHit();
+	}
+
+	FThresholdGuard Guard([this]()
+	{
+		UpdatePixelThreshold(ConeElement, PixelHitDistanceThreshold);
+		UpdatePixelThreshold(BoxElement, PixelHitDistanceThreshold);
+		UpdatePixelThreshold(SphereElement, PixelHitDistanceThreshold);
+	});
+
+	// update threshold if hitting the mask
+	if (UGizmoElementBase* HitMaskGizmo = HitMask.Get())
+	{
+		const FInputRayHit MaskHit = HitMaskGizmo->LineTrace(ViewContext, LineTraceState, RayOrigin, RayDirection);
+		if (MaskHit.bHit)
+		{
+			static constexpr float NoPixelHitThreshold = 0.0f;
+			UpdatePixelThreshold(ConeElement, NoPixelHitThreshold);
+			UpdatePixelThreshold(BoxElement, NoPixelHitThreshold);
+			UpdatePixelThreshold(SphereElement, NoPixelHitThreshold);
+		}
+	}
+
+	auto LineTraceElement = [&](UGizmoElementBase* InElement) -> FInputRayHit
+	{
+		check(InElement);
+		return InElement->LineTrace(ViewContext, CurrentLineTraceState, RayOrigin, RayDirection);
+	};
+
+	FInputRayHit Hit;
+	if (Type == EGizmoElementArrowHeadType::Cone)
+	{
+		Hit = LineTraceElement(ConeElement);
+	}
+	else if (Type == EGizmoElementArrowHeadType::Cube)
+	{
+		Hit = LineTraceElement(BoxElement);
+	}
+	else if (Type == EGizmoElementArrowHeadType::Sphere)
+	{
+		Hit = LineTraceElement(SphereElement);
+	}
+
+	if (Hit.bHit)
+	{
+		Hit.SetHitObject(this);
+		Hit.HitIdentifier = PartIdentifier;
+	}
+
+	return Hit;
+}
+
+void UGizmoElementArrowHead::SetCenter(const FVector& InCenter)
+{
+	if (Center != InCenter)
+	{
+		Center = InCenter;
+		bUpdate = true;
+	}
+}
+
+FVector UGizmoElementArrowHead::GetCenter() const
+{
+	return Center;
+}
+
+void UGizmoElementArrowHead::SetDirection(const FVector& InDirection)
+{
+	Direction = InDirection;
+	Direction.Normalize();
+	bUpdate = true;
+}
+
+FVector UGizmoElementArrowHead::GetDirection() const
+{
+	return Direction;
+}
+
+void UGizmoElementArrowHead::SetSideDirection(const FVector& InSideDirection)
+{
+	SideDirection = InSideDirection;
+	SideDirection.Normalize();
+	bUpdate = true;
+}
+
+FVector UGizmoElementArrowHead::GetSideDirection() const
+{
+	return SideDirection;
+}
+
+void UGizmoElementArrowHead::SetLength(float InLength)
+{
+	if (Length != InLength)
+	{
+		Length = InLength;
+		bUpdate = true;
+	}
+}
+
+float UGizmoElementArrowHead::GetLength() const
+{
+	return Length;
+}
+
+void UGizmoElementArrowHead::SetRadius(float InRadius)
+{
+	if (Radius != InRadius)
+	{
+		Radius = InRadius;
+		bUpdate = true;
+	}
+}
+
+float UGizmoElementArrowHead::GetRadius() const
+{
+	return Radius;
+}
+
+void UGizmoElementArrowHead::SetNumSides(int32 InNumSides)
+{
+	if (NumSides != InNumSides)
+    {
+    	NumSides = InNumSides;
+    	bUpdate = true;
+    }
+}
+
+int32 UGizmoElementArrowHead::GetNumSides() const
+{
+	return NumSides;
+}
+
+void UGizmoElementArrowHead::SetType(EGizmoElementArrowHeadType InHeadType)
+{
+	if (InHeadType != Type)
+	{
+		// Removes the given element from the group, if the element is valid
+		auto RemoveElementFromGroup = [&](UGizmoElementBase* InElement)
+		{
+			if (!InElement)
+			{
+				return;
+			}
+
+			UGizmoElementArrowHead::Remove(InElement);
+			InElement = nullptr;
+		};
+
+		Type = InHeadType;
+
+		if (Type == EGizmoElementArrowHeadType::Cone)
+		{
+			ConeElement = NewObject<UGizmoElementCone>();
+			UGizmoElementArrowHead::Add(ConeElement);
+
+			RemoveElementFromGroup(BoxElement);
+			RemoveElementFromGroup(SphereElement);
+		}
+		else if (Type == EGizmoElementArrowHeadType::Cube)
+		{
+			BoxElement = NewObject<UGizmoElementBox>();
+			UGizmoElementArrowHead::Add(BoxElement);
+
+			RemoveElementFromGroup(ConeElement);
+			RemoveElementFromGroup(SphereElement);
+		}
+		else if (Type == EGizmoElementArrowHeadType::Sphere)
+		{
+			SphereElement = NewObject<UGizmoElementSphere>();
+			UGizmoElementArrowHead::Add(SphereElement);
+
+			RemoveElementFromGroup(BoxElement);
+			RemoveElementFromGroup(ConeElement);
+		}
+
+		Update();
+	}
+}
+
+EGizmoElementArrowHeadType UGizmoElementArrowHead::GetType() const
+{
+	return Type;
+}
+
+void UGizmoElementArrowHead::SetPixelHitDistanceThreshold(float InPixelHitDistanceThreshold)
+{
+	if (PixelHitDistanceThreshold != InPixelHitDistanceThreshold)
+	{
+		PixelHitDistanceThreshold = InPixelHitDistanceThreshold;
+		bUpdate = true;
+	}
+}
+
+void UGizmoElementArrowHead::SetHitMask(const TWeakObjectPtr<UGizmoElementBase>& InHitMask)
+{
+	HitMask = InHitMask;
+}
+
+void UGizmoElementArrowHead::Update()
+{
+	static constexpr bool bEndCaps = true;
+
+	auto UpdateCommon = [&](UGizmoElementBase* InElement)
+	{
+		InElement->SetVertexColor(GetVertexColor());
+		InElement->SetHoverVertexColor(GetHoverVertexColor());
+		InElement->SetInteractVertexColor(GetInteractVertexColor());
+		InElement->SetSelectVertexColor(GetSelectVertexColor());
+		InElement->SetSubdueVertexColor(GetSubdueVertexColor());
+
+		InElement->SetMaterial(MeshRenderAttributes.Material.Value);
+		InElement->SetHoverMaterial(MeshRenderAttributes.HoverMaterial.Value);
+		InElement->SetInteractMaterial(MeshRenderAttributes.InteractMaterial.Value);
+		InElement->SetSelectMaterial(MeshRenderAttributes.SelectMaterial.Value);
+		InElement->SetSubdueMaterial(MeshRenderAttributes.SubdueMaterial.Value);
+	};
+
+	if (Type == EGizmoElementArrowHeadType::Cone)
+	{
+		if (ensure(ConeElement))
+		{
+			ConeElement->SetOrigin(Direction * (Length * 0.9f));
+			ConeElement->SetDirection(-Direction);
+			ConeElement->SetHeight(Length);
+			ConeElement->SetRadius(Radius);
+			ConeElement->SetNumSides(NumSides);
+			ConeElement->SetElementInteractionState(ElementInteractionState);
+			ConeElement->SetPixelHitDistanceThreshold(PixelHitDistanceThreshold);
+			ConeElement->SetEndCaps(bEndCaps);
+
+			UpdateCommon(ConeElement);
+		}
+	}
+	else if (Type == EGizmoElementArrowHeadType::Cube)
+	{
+		if (ensure(BoxElement))
+		{
+			BoxElement->SetCenter(FVector::ZeroVector);
+			BoxElement->SetUpDirection(Direction);
+			BoxElement->SetSideDirection(SideDirection);
+			BoxElement->SetDimensions(FVector(Length, Length, Length));
+			BoxElement->SetElementInteractionState(ElementInteractionState);
+			BoxElement->SetPixelHitDistanceThreshold(PixelHitDistanceThreshold);
+
+			UpdateCommon(BoxElement);
+		}
+	}
+	else if (Type == EGizmoElementArrowHeadType::Sphere)
+	{
+		if (ensure(SphereElement))
+		{
+			SphereElement->SetCenter(FVector::ZeroVector);
+			SphereElement->SetRadius(Radius);
+			SphereElement->SetNumSides(NumSides);
+			SphereElement->SetElementInteractionState(ElementInteractionState);
+			SphereElement->SetPixelHitDistanceThreshold(PixelHitDistanceThreshold);
+
+			UpdateCommon(SphereElement);
+		}
+	}
+
+	bUpdate = false;
+}
diff --git a/Engine/Source/Runtime/InteractiveToolsFramework/Public/BaseGizmos/GizmoElementArrow.h b/Engine/Source/Runtime/InteractiveToolsFramework/Public/BaseGizmos/GizmoElementArrow.h
--- a/Engine/Source/Runtime/InteractiveToolsFramework/Public/BaseGizmos/GizmoElementArrow.h
+++ b/Engine/Source/Runtime/InteractiveToolsFramework/Public/BaseGizmos/GizmoElementArrow.h
@@ -3,30 +3,25 @@
 #pragma once
 
 #include "CoreMinimal.h"
+#include "BaseGizmos/GizmoElementArrowHead.h"
 #include "BaseGizmos/GizmoElementBase.h"
+#include "GizmoElementGroup.h"
 #include "InputState.h"
 #include "UObject/ObjectMacros.h"
+
 #include "GizmoElementArrow.generated.h"
 
 class UGizmoElementBox;
 class UGizmoElementCone;
 class UGizmoElementCylinder;
+class UGizmoElementSphere;
 
 /**
  * Simple object intended to be used as part of 3D Gizmos.
  * Draws a 3D arrow based on parameters.
  */
-
-UENUM()
-enum class EGizmoElementArrowHeadType
-{
-	Cone,
-	Cube
-};
-
-
 UCLASS(Transient, MinimalAPI)
-class UGizmoElementArrow : public UGizmoElementBase
+class UGizmoElementArrow : public UGizmoElementGroupBase
 {
 	GENERATED_BODY()
 
@@ -110,13 +105,9 @@ protected:
 	UPROPERTY()
 	TObjectPtr<UGizmoElementCylinder> CylinderElement;
 
-	// Arrow cone head
-	UPROPERTY()
-	TObjectPtr<UGizmoElementCone> ConeElement;
-
-	// Arrow box head
+	// Arrow head
 	UPROPERTY()
-	TObjectPtr<UGizmoElementBox> BoxElement;
+	TObjectPtr<UGizmoElementArrowHead> HeadElement;
 
 	// Location of center of arrow base circle.
 	UPROPERTY()
diff --git a/Engine/Source/Runtime/InteractiveToolsFramework/Public/BaseGizmos/GizmoElementArrowHead.h b/Engine/Source/Runtime/InteractiveToolsFramework/Public/BaseGizmos/GizmoElementArrowHead.h
new file mode 100644
--- /dev/null
+++ b/Engine/Source/Runtime/InteractiveToolsFramework/Public/BaseGizmos/GizmoElementArrowHead.h
@@ -0,0 +1,128 @@
+// Copyright Epic Games, Inc. All Rights Reserved.
+
+#pragma once
+
+#include "GizmoElementBase.h"
+#include "GizmoElementGroup.h"
+
+#include "GizmoElementArrowHead.generated.h"
+
+class UGizmoElementBox;
+class UGizmoElementCone;
+class UGizmoElementSphere;
+
+UENUM()
+enum class EGizmoElementArrowHeadType
+{
+	None,
+	Cone,
+	Cube,
+	Sphere
+};
+
+/**
+ * Simple object intended to be used as part of 3D Gizmos.
+ * Draws a 3D arrowhead based on parameters.
+ */
+UCLASS(Transient, MinimalAPI)
+class UGizmoElementArrowHead : public UGizmoElementGroupBase
+{
+	GENERATED_BODY()
+
+public:
+	INTERACTIVETOOLSFRAMEWORK_API UGizmoElementArrowHead();
+
+	//~ Begin UGizmoElementBase Interface.
+	INTERACTIVETOOLSFRAMEWORK_API virtual void Render(IToolsContextRenderAPI* RenderAPI, const FRenderTraversalState& RenderState) override;
+	INTERACTIVETOOLSFRAMEWORK_API virtual FInputRayHit LineTrace(const UGizmoViewContext* ViewContext, const FLineTraceTraversalState& LineTraceState, const FVector& RayOrigin, const FVector& RayDirection) override;
+	//~ End UGizmoElementBase Interface.
+
+	// Location of center of the arrow head
+	INTERACTIVETOOLSFRAMEWORK_API virtual void SetCenter(const FVector& InCenter);
+	INTERACTIVETOOLSFRAMEWORK_API virtual FVector GetCenter() const;
+
+	// Arrow direction.
+	INTERACTIVETOOLSFRAMEWORK_API virtual void SetDirection(const FVector& InDirection);
+	INTERACTIVETOOLSFRAMEWORK_API virtual FVector GetDirection() const;
+
+	// Arrow side direction for box head.
+	INTERACTIVETOOLSFRAMEWORK_API virtual void SetSideDirection(const FVector& InSideDirection);
+	INTERACTIVETOOLSFRAMEWORK_API virtual FVector GetSideDirection() const;
+
+	// Arrow head length, used for both cone and cube head.
+	INTERACTIVETOOLSFRAMEWORK_API virtual void SetLength(float InLength);
+	INTERACTIVETOOLSFRAMEWORK_API virtual float GetLength() const;
+
+	// Arrow head radius, if cone or sphere.
+	INTERACTIVETOOLSFRAMEWORK_API virtual void SetRadius(float InRadius);
+	INTERACTIVETOOLSFRAMEWORK_API virtual float GetRadius() const;
+
+	// Number of sides for cylinder and cone, if relevant.
+	INTERACTIVETOOLSFRAMEWORK_API virtual void SetNumSides(int32 InNumSides);
+	INTERACTIVETOOLSFRAMEWORK_API virtual int32 GetNumSides() const;
+
+	// Head type cone or cube.
+	INTERACTIVETOOLSFRAMEWORK_API virtual void SetType(EGizmoElementArrowHeadType InHeadType);
+	INTERACTIVETOOLSFRAMEWORK_API virtual EGizmoElementArrowHeadType GetType() const;
+
+	// Pixel hit distance threshold, element will be scaled enough to add this threshold when line-tracing.
+	INTERACTIVETOOLSFRAMEWORK_API virtual void SetPixelHitDistanceThreshold(float InPixelHitDistanceThreshold) override;
+
+	// Hit Mask Gizmo element used for adaptive pixel hit threshold
+	INTERACTIVETOOLSFRAMEWORK_API void SetHitMask(const TWeakObjectPtr<UGizmoElementBase>& InHitMask);
+
+protected:
+
+	// Update the appropriate shape based on parameters
+	INTERACTIVETOOLSFRAMEWORK_API virtual void Update();
+
+protected:
+
+	// Flag indicating properties need to be updated prior to render
+	bool bUpdate = true;
+
+	// Cone head
+	UPROPERTY()
+	TObjectPtr<UGizmoElementCone> ConeElement;
+
+	// Box head
+	UPROPERTY()
+	TObjectPtr<UGizmoElementBox> BoxElement;
+
+	// Sphere head
+	UPROPERTY()
+	TObjectPtr<UGizmoElementSphere> SphereElement;
+
+	// Location of center
+	UPROPERTY()
+	FVector Center = FVector::ZeroVector;
+
+	// Direction of arrow axis
+	UPROPERTY()
+	FVector Direction = FVector(0.0f, 0.0f, 1.0f);
+
+	// Side direction for box
+	UPROPERTY()
+	FVector SideDirection = FVector(0.0f, 1.0f, 0.0f);
+
+	// Length of head, cone or box
+	UPROPERTY()
+	float Length = 0.5f;
+
+	// Radius of head cone
+	UPROPERTY()
+	float Radius = 0.75f;
+
+	// Number of sides for tessellating cone and cylinder
+	UPROPERTY()
+	int32 NumSides = 32;
+
+	// Head type
+	UPROPERTY()
+	EGizmoElementArrowHeadType Type = EGizmoElementArrowHeadType::Cone;
+
+	// Hit mask
+	UPROPERTY()
+	TWeakObjectPtr<UGizmoElementBase> HitMask = nullptr;
+};
+
