ITF - Fix ToolBuilderUtil::FindFirstComponentOfType - This unused function does not actually work because the return type will be incorrect for all values of ComponentType which are derived from UActorComponent.

The logic for dynamically checking that the found component is of the specified type, but the return type of FindFirstComponent is UActorComponent* which is not equivalent to ComponentType* for non-trivial calls to this function. We can safely static cast because we only get a pointer back if the dynamic cast succeeded.

diff --git a/Engine/Source/Runtime/InteractiveToolsFramework/Public/ToolBuilderUtil.h b/Engine/Source/Runtime/InteractiveToolsFramework/Public/ToolBuilderUtil.h
--- a/Engine/Source/Runtime/InteractiveToolsFramework/Public/ToolBuilderUtil.h
+++ b/Engine/Source/Runtime/InteractiveToolsFramework/Public/ToolBuilderUtil.h
@@ -79,7 +79,7 @@ int ToolBuilderUtil::CountSelectedComponentsOfType(const FToolBuilderState& Inpu
 template<typename ComponentType>
 ComponentType* ToolBuilderUtil::FindFirstComponentOfType(const FToolBuilderState& InputState)
 {
-	return FindFirstComponent(InputState, [](UActorComponent* Component) { return Cast<ComponentType>(Component) != nullptr; });
+	return static_cast<ComponentType*>(FindFirstComponent(InputState, [](UActorComponent* Component) { return Cast<ComponentType>(Component) != nullptr; }));
 }
 
 template<typename ActorType>
