[ITF] Added non-const GetMaterial to FGizmoElementMaterialAttribute, to use the return value in other calls/setter - almost all of which require a non-const pointer.

diff --git a/Engine/Source/Runtime/InteractiveToolsFramework/Private/BaseGizmos/GizmoElementRenderState.cpp b/Engine/Source/Runtime/InteractiveToolsFramework/Private/BaseGizmos/GizmoElementRenderState.cpp
--- a/Engine/Source/Runtime/InteractiveToolsFramework/Private/BaseGizmos/GizmoElementRenderState.cpp
+++ b/Engine/Source/Runtime/InteractiveToolsFramework/Private/BaseGizmos/GizmoElementRenderState.cpp
@@ -28,7 +28,7 @@ void FGizmoElementColorAttribute::Reset()
 	bOverridesChildState = false;
 }
 
-void  FGizmoElementColorAttribute::UpdateState(const FGizmoElementColorAttribute& InChildColorAttribute)
+void FGizmoElementColorAttribute::UpdateState(const FGizmoElementColorAttribute& InChildColorAttribute)
 {
 	if (InChildColorAttribute.bHasValue && !(bOverridesChildState && bHasValue))
 	{
@@ -47,6 +47,15 @@ const UMaterialInterface* FGizmoElementMaterialAttribute::GetMaterial() const
 	return nullptr;
 }
 
+UMaterialInterface* FGizmoElementMaterialAttribute::GetMaterial()
+{
+	if (Value.IsValid())
+	{
+		return Value.Get();
+	}
+	return nullptr;
+}
+
 void FGizmoElementMaterialAttribute::SetMaterial(TWeakObjectPtr<UMaterialInterface> InMaterial, bool InOverridesChildState)
 {
 	Value = InMaterial;
diff --git a/Engine/Source/Runtime/InteractiveToolsFramework/Public/BaseGizmos/GizmoElementRenderState.h b/Engine/Source/Runtime/InteractiveToolsFramework/Public/BaseGizmos/GizmoElementRenderState.h
--- a/Engine/Source/Runtime/InteractiveToolsFramework/Public/BaseGizmos/GizmoElementRenderState.h
+++ b/Engine/Source/Runtime/InteractiveToolsFramework/Public/BaseGizmos/GizmoElementRenderState.h
@@ -72,6 +72,9 @@ struct FGizmoElementMaterialAttribute
 	// Return the material value.
 	INTERACTIVETOOLSFRAMEWORK_API const UMaterialInterface* GetMaterial() const;
 
+	// Return the material value.
+	INTERACTIVETOOLSFRAMEWORK_API UMaterialInterface* GetMaterial();
+
 	// Set material value and override flag
 	INTERACTIVETOOLSFRAMEWORK_API void SetMaterial(TWeakObjectPtr<UMaterialInterface> InColor, bool InOverridesChildState = false);
 
