[ITF] Added WorldToPixel convenience function to GizmoViewContext - alongside existing WorldToSccreen, ScreenToPixel

diff --git a/Engine/Source/Runtime/InteractiveToolsFramework/Public/BaseGizmos/GizmoViewContext.h b/Engine/Source/Runtime/InteractiveToolsFramework/Public/BaseGizmos/GizmoViewContext.h
--- a/Engine/Source/Runtime/InteractiveToolsFramework/Public/BaseGizmos/GizmoViewContext.h
+++ b/Engine/Source/Runtime/InteractiveToolsFramework/Public/BaseGizmos/GizmoViewContext.h
@@ -70,6 +70,18 @@ public:
 		return ViewMatrices.GetViewProjectionMatrix().TransformFVector4(FVector4(WorldPoint, 1));
 	}
 
+	// WorldToScreen -> ScreenToPixel. A result of (-1, -1) indicates invalid W.
+	FVector2D WorldToPixel(const FVector& WorldPoint) const
+	{
+		FVector2D ScreenPos;
+		if (ScreenToPixel(WorldToScreen(WorldPoint), ScreenPos))
+		{
+			return ScreenPos;
+		}
+
+		return -FVector2D::One();
+	}
+
 	bool ScreenToPixel(const FVector4& ScreenPoint, FVector2D& OutPixelLocation) const
 	{
 		// implementation copied from FSceneView::ScreenToPixel() 
