[ITF] Added DrawGrid to ToolDataVisualizer

diff --git a/Engine/Source/Runtime/InteractiveToolsFramework/Private/ToolDataVisualizer.cpp b/Engine/Source/Runtime/InteractiveToolsFramework/Private/ToolDataVisualizer.cpp
--- a/Engine/Source/Runtime/InteractiveToolsFramework/Private/ToolDataVisualizer.cpp
+++ b/Engine/Source/Runtime/InteractiveToolsFramework/Private/ToolDataVisualizer.cpp
@@ -166,6 +166,47 @@ void FToolDataVisualizer::InternalDrawSquare(const FVector& Center, const FVecto
 	InternalDrawTransformedLine(C00, C10, Color, LineThicknessIn, bDepthTestedIn);
 }
 
+// @note: This is taken from MeshDebugDraw::DrawSimpleGrid
+void FToolDataVisualizer::InternalDrawGrid(const FVector& Center, const FVector& Normal, const FVector& Size, const FLinearColor& Color, float LineThicknessIn, bool bDepthTestedIn)
+{
+	const ESceneDepthPriorityGroup DepthPriority = (bDepthTestedIn) ? SDPG_World : SDPG_Foreground;
+
+	FVector Tan1, Tan2;
+	GizmoMath::MakeNormalPlaneBasis(Normal, Tan1, Tan2);
+
+	const FVector CC = TransformP(Center);
+	const FVector SA = TransformV(Tan1);
+	const FVector SB = TransformV(Tan2);
+
+	// Always odd
+	constexpr int32 NumGridLines = (FMath::DivideAndRoundUp(16, 2) * 2) + 1;
+	const float GridLineSpacing = Size.X / (NumGridLines - 1);
+
+	const double Width = static_cast<double>(NumGridLines - 1) * GridLineSpacing;
+	const double Extent = Width * 0.5;
+
+	constexpr int32 LineSteps = FMath::DivideAndRoundUp(NumGridLines, 2);
+	for (int32 LineIndex = 0; LineIndex < LineSteps; ++LineIndex)
+	{
+		const double DistanceX = static_cast<double>(LineIndex) * GridLineSpacing;
+		FVector3d A = CC - Extent * SB - DistanceX * SA;
+		FVector3d B = CC + Extent * SB - DistanceX * SA;
+		CurrentPDI->DrawLine(A, B, Color, DepthPriority, LineThicknessIn, 0, true);
+
+		A = CC - Extent * SB + DistanceX * SA;
+		B = CC + Extent * SB + DistanceX * SA;
+		CurrentPDI->DrawLine(A, B, Color, DepthPriority, LineThicknessIn, 0, true);
+
+		A = CC - Extent * SA - DistanceX * SB;
+		B = CC + Extent * SA - DistanceX * SB;
+		CurrentPDI->DrawLine(A, B, Color, DepthPriority, LineThicknessIn, 0, true);
+
+		A = CC - Extent * SA + DistanceX * SB;
+		B = CC + Extent * SA + DistanceX * SB;
+		CurrentPDI->DrawLine(A, B, Color, DepthPriority, LineThicknessIn, 0, true);
+	}
+}
+
 void FToolDataVisualizer::InternalDrawWireCylinder(const FVector& Position, const FVector& Normal, float Radius, float Height, int Steps, const FLinearColor& Color, float LineThicknessIn, bool bDepthTestedIn)
 {
 	FVector Tan1, Tan2;
diff --git a/Engine/Source/Runtime/InteractiveToolsFramework/Public/ToolDataVisualizer.h b/Engine/Source/Runtime/InteractiveToolsFramework/Public/ToolDataVisualizer.h
--- a/Engine/Source/Runtime/InteractiveToolsFramework/Public/ToolDataVisualizer.h
+++ b/Engine/Source/Runtime/InteractiveToolsFramework/Public/ToolDataVisualizer.h
@@ -96,7 +96,7 @@ public:
 
 
 	/** Update the default line color and thickness */
-	void SetLineParameters(const FLinearColor& Color, float Thickness)
+	void SetLineParameters(const FLinearColor& Color, float Thickness, float InArrowHeadSize = 4.0f)
 	{
 		LineColor = Color;
 		LineThickness = Thickness;
@@ -143,7 +143,6 @@ public:
 		InternalDrawTransformedLine(TransformP((FVector)A), TransformP((FVector)B), Color, LineThicknessIn, bDepthTestedIn);
 	}
 
-
 	/** Draw a point with default parameters */
 	template<typename PointType>
 	void DrawPoint(const PointType& Position)
@@ -245,6 +244,12 @@ public:
 		InternalDrawSquare(Center, SideA, SideB, LineColor, LineThickness, bDepthTested);
 	}
 
+	template<typename PointType>
+	void DrawGrid(const PointType& Center, const PointType& Normal, const PointType& Size, const FLinearColor& Color, float LineThicknessIn, bool bDepthTestedIn)
+	{
+		InternalDrawGrid(Center, Normal, Size, Color, LineThicknessIn, bDepthTestedIn);
+	}
+
 	/** Draw a 3D disc at given position/normal with the given parameters. Tangent axes are defined internally */
 	template<typename PointType>
 	void DrawDisc(const PointType& Position, const PointType& Normal, float Radius, int Steps, const FColor& Color, FMaterialRenderProxy* RenderProxy, bool bDepthTestedIn)
@@ -287,6 +292,8 @@ protected:
 	INTERACTIVETOOLSFRAMEWORK_API virtual void InternalDrawWireBox(const FBox& Box, const FLinearColor& Color, float LineThickness, bool bDepthTested);
 	// actually does the square drawing
 	INTERACTIVETOOLSFRAMEWORK_API virtual void InternalDrawSquare(const FVector& Center, const FVector& SideA, const FVector& SideB, const FLinearColor& Color, float LineThickness, bool bDepthTested);
+	// actually does the grid drawing
+	INTERACTIVETOOLSFRAMEWORK_API virtual void InternalDrawGrid(const FVector& Center, const FVector& Normal, const FVector& Size, const FLinearColor& Color, float LineThickness, bool bDepthTested);
 	// actually does the X drawing
 	INTERACTIVETOOLSFRAMEWORK_API virtual void InternalDrawViewFacingX(const FVector& Position, float Width, const FLinearColor& Color, float LineThickness, bool bDepthTested);
 	// actually does the disc drawing
