[ITF] Fixed normal direction on GizmoElementCircle

diff --git a/Engine/Source/Runtime/InteractiveToolsFramework/Private/BaseGizmos/GizmoElementCircle.cpp b/Engine/Source/Runtime/InteractiveToolsFramework/Private/BaseGizmos/GizmoElementCircle.cpp
--- a/Engine/Source/Runtime/InteractiveToolsFramework/Private/BaseGizmos/GizmoElementCircle.cpp
+++ b/Engine/Source/Runtime/InteractiveToolsFramework/Private/BaseGizmos/GizmoElementCircle.cpp
@@ -1,8 +1,10 @@
 // Copyright Epic Games, Inc. All Rights Reserved.
  
 #include "BaseGizmos/GizmoElementCircle.h"
-#include "BaseGizmos/GizmoRenderingUtil.h"
+
 #include "BaseGizmos/GizmoMath.h"
+#include "BaseGizmos/GizmoRenderingUtil.h"
+#include "DynamicMeshBuilder.h"
 #include "InputState.h"
 #include "Intersection/IntersectionUtil.h"
 #include "Materials/MaterialInterface.h"
@@ -10,13 +12,71 @@
 
 #include UE_INLINE_GENERATED_CPP_BY_NAME(GizmoElementCircle)
 
+namespace UE::InteractiveToolsFramework::Private
+{
+	namespace GizmoElementCircleLocals
+	{
+		// copy/paste of DrawDisc from PrimitiveDrawingUtils.h, with corrected (flat) normals
+		void DrawDisc(FPrimitiveDrawInterface* PDI, const FVector& Base, const FVector& XAxis, const FVector& YAxis, FColor Color, double Radius, int32 NumSides, const FMaterialRenderProxy* MaterialRenderProxy, uint8 DepthPriority)
+		{
+			check(NumSides >= 3);
+
+			const float AngleDelta = 2.0f * UE_PI / NumSides;
+
+			const FVector2D TC = FVector2D(0.0f, 0.0f);
+			const float TCStep = 1.0f / NumSides;
+
+			const FVector ZAxis = (XAxis) ^ YAxis;
+
+			const FVector3f TangentX = static_cast<FVector3f>(XAxis);
+			const FVector3f TangentY = static_cast<FVector3f>(-YAxis);
+			const FVector3f TangentZ = static_cast<FVector3f>(ZAxis);
+
+			FDynamicMeshBuilder MeshBuilder(PDI->View->GetFeatureLevel());
+
+			//Compute vertices for base circle.
+			for (int32 SideIndex = 0; SideIndex < NumSides; SideIndex++)
+			{
+				const FVector Vertex = Base + (XAxis * FMath::Cos(AngleDelta * (SideIndex)) + YAxis * FMath::Sin(AngleDelta * (SideIndex))) * Radius;
+
+				FDynamicMeshVertex MeshVertex;
+				MeshVertex.Position = static_cast<FVector3f>(Vertex);
+				MeshVertex.Color = Color;
+				MeshVertex.TextureCoordinate[0] = FVector2f(TC);
+				MeshVertex.TextureCoordinate[0].X += TCStep * SideIndex;
+
+				MeshVertex.SetTangents(
+					TangentX,
+					TangentY,
+					TangentZ
+				);
+
+				MeshBuilder.AddVertex(MeshVertex); //Add bottom vertex
+			}
+
+			//Add top/bottom triangles, in the style of a fan.
+			for (int32 SideIndex = 0; SideIndex < NumSides - 1; SideIndex++)
+			{
+				int32 V0 = 0;
+				int32 V1 = SideIndex;
+				int32 V2 = (SideIndex + 1);
+
+				MeshBuilder.AddTriangle(V0, V1, V2);
+				MeshBuilder.AddTriangle(V0, V2, V1);
+			}
+
+			MeshBuilder.Draw(PDI, FMatrix::Identity, MaterialRenderProxy, DepthPriority, 0.f);
+		}
+	}
+}
+
 void UGizmoElementCircle::Render(IToolsContextRenderAPI* RenderAPI, const FRenderTraversalState& RenderState)
 {
 	if (!ensure(RenderAPI))
 	{
 		return;
 	}
-	
+
 	FRenderTraversalState CurrentRenderState(RenderState);
 	bool bVisibleViewDependent = UpdateRenderState(RenderAPI, Center, CurrentRenderState);
 
@@ -33,8 +93,10 @@ void UGizmoElementCircle::Render(IToolsContextRenderAPI* RenderAPI, const FRende
 		{
 			if (const UMaterialInterface* UseMaterial = CurrentRenderState.GetCurrentMaterial())
 			{
-				FColor VertexColor = CurrentRenderState.GetCurrentVertexColor().ToFColor(false);
-				DrawDisc(PDI, WorldCenter, WorldAxis0, WorldAxis1, VertexColor, WorldRadius, NumSegments, UseMaterial->GetRenderProxy(), SDPG_Foreground);
+				using namespace UE::InteractiveToolsFramework::Private;
+
+				const FColor VertexColor = CurrentRenderState.GetCurrentVertexColor().ToFColor(true);
+				GizmoElementCircleLocals::DrawDisc(PDI, WorldCenter, WorldAxis0, WorldAxis1, VertexColor, WorldRadius, NumSegments, UseMaterial->GetRenderProxy(), SDPG_Foreground);
 			}
 		}
 
@@ -162,5 +224,3 @@ bool UGizmoElementCircle::GetHitLine() const
 {
 	return bHitLine;
 }
-
-
