[ITF] Added non-const material getters to GizmoElements

diff --git a/Engine/Source/Runtime/InteractiveToolsFramework/Private/BaseGizmos/GizmoElementBase.cpp b/Engine/Source/Runtime/InteractiveToolsFramework/Private/BaseGizmos/GizmoElementBase.cpp
--- a/Engine/Source/Runtime/InteractiveToolsFramework/Private/BaseGizmos/GizmoElementBase.cpp
+++ b/Engine/Source/Runtime/InteractiveToolsFramework/Private/BaseGizmos/GizmoElementBase.cpp
@@ -622,6 +622,11 @@ const UMaterialInterface* UGizmoElementBase::GetMaterial() const
 	return MeshRenderAttributes.Material.GetMaterial();
 }
 
+UMaterialInterface* UGizmoElementBase::GetMaterial()
+{
+	return MeshRenderAttributes.Material.GetMaterial();
+}
+
 bool UGizmoElementBase::DoesMaterialOverrideChildState() const
 {
 	return MeshRenderAttributes.Material.bOverridesChildState;
@@ -642,6 +647,11 @@ const UMaterialInterface* UGizmoElementBase::GetHoverMaterial() const
 	return MeshRenderAttributes.HoverMaterial.GetMaterial();
 }
 
+UMaterialInterface* UGizmoElementBase::GetHoverMaterial()
+{
+	return MeshRenderAttributes.HoverMaterial.GetMaterial();
+}
+
 bool UGizmoElementBase::DoesHoverMaterialOverrideChildState() const
 {
 	return MeshRenderAttributes.HoverMaterial.bOverridesChildState;
@@ -662,6 +672,11 @@ const UMaterialInterface* UGizmoElementBase::GetInteractMaterial() const
 	return MeshRenderAttributes.InteractMaterial.GetMaterial();
 }
 
+UMaterialInterface* UGizmoElementBase::GetInteractMaterial()
+{
+	return MeshRenderAttributes.InteractMaterial.GetMaterial();
+}
+
 bool UGizmoElementBase::DoesInteractMaterialOverrideChildState() const
 {
 	return MeshRenderAttributes.InteractMaterial.bOverridesChildState;
@@ -682,6 +697,11 @@ const UMaterialInterface* UGizmoElementBase::GetSelectMaterial() const
 	return MeshRenderAttributes.SelectMaterial.GetMaterial();
 }
 
+UMaterialInterface* UGizmoElementBase::GetSelectMaterial()
+{
+	return MeshRenderAttributes.SelectMaterial.GetMaterial();
+}
+
 bool UGizmoElementBase::DoesSelectMaterialOverrideChildState() const
 {
 	return MeshRenderAttributes.SelectMaterial.bOverridesChildState;
@@ -702,6 +722,11 @@ const UMaterialInterface* UGizmoElementBase::GetSubdueMaterial() const
 	return MeshRenderAttributes.SubdueMaterial.GetMaterial();
 }
 
+UMaterialInterface* UGizmoElementBase::GetSubdueMaterial()
+{
+	return MeshRenderAttributes.SubdueMaterial.GetMaterial();
+}
+
 bool UGizmoElementBase::DoesSubdueMaterialOverrideChildState() const
 {
 	return MeshRenderAttributes.SubdueMaterial.bOverridesChildState;
diff --git a/Engine/Source/Runtime/InteractiveToolsFramework/Public/BaseGizmos/GizmoElementBase.h b/Engine/Source/Runtime/InteractiveToolsFramework/Public/BaseGizmos/GizmoElementBase.h
--- a/Engine/Source/Runtime/InteractiveToolsFramework/Public/BaseGizmos/GizmoElementBase.h
+++ b/Engine/Source/Runtime/InteractiveToolsFramework/Public/BaseGizmos/GizmoElementBase.h
@@ -564,6 +564,21 @@ protected:
 	// @param OutAlignRot the rotation to align this element in local space, should be prepended to the local-to-world transform.
 	INTERACTIVETOOLSFRAMEWORK_API virtual bool GetViewAlignRot(const FVector& InViewLocation, const FVector& InViewDirection, const FVector& InViewUp, bool bInPerspectiveView, const FTransform& InLocalToWorldTransform, const FVector& InLocalCenter, FQuat& OutAlignRot) const;
 
+	// Get mesh render state material attribute's value. 
+	INTERACTIVETOOLSFRAMEWORK_API UMaterialInterface* GetMaterial();
+
+	// Get mesh render state hover material attribute's value.
+	INTERACTIVETOOLSFRAMEWORK_API UMaterialInterface* GetHoverMaterial();
+
+	// Get mesh render state interact material attribute's value. 
+	INTERACTIVETOOLSFRAMEWORK_API UMaterialInterface* GetInteractMaterial();
+
+	// Get mesh render state Select material attribute's value. 
+	INTERACTIVETOOLSFRAMEWORK_API UMaterialInterface* GetSelectMaterial();
+
+	// Get mesh render state Subdue material attribute's value. 
+	INTERACTIVETOOLSFRAMEWORK_API UMaterialInterface* GetSubdueMaterial();
+
 	// Update render state during render traversal, determines the current render state for this element 
 	// @param RenderAPI - tools render context
 	// @param InLocalCenter - local element center position
