[ITF] Made FLinearColor->FColor SRGB conversions consistent (all true). Ensures vertex-colors match colors defined at the material level.

diff --git a/Engine/Source/Runtime/InteractiveToolsFramework/Private/BaseGizmos/GizmoElementArc.cpp b/Engine/Source/Runtime/InteractiveToolsFramework/Private/BaseGizmos/GizmoElementArc.cpp
--- a/Engine/Source/Runtime/InteractiveToolsFramework/Private/BaseGizmos/GizmoElementArc.cpp
+++ b/Engine/Source/Runtime/InteractiveToolsFramework/Private/BaseGizmos/GizmoElementArc.cpp
@@ -173,8 +173,8 @@ void UGizmoElementArc::Render(IToolsContextRenderAPI* RenderAPI, const FRenderTr
 
 	const double WorldOuterRadius = Radius * CurrentRenderState.LocalToWorldTransform.GetScale3D().X;
 	const double WorldInnerRadius = InnerRadius * CurrentRenderState.LocalToWorldTransform.GetScale3D().X;
-	const FColor VertexColor = CurrentRenderState.GetCurrentVertexColor().ToFColor(false);
-	const FColor LineColor = CurrentRenderState.GetCurrentLineColor().ToFColor(false);
+	const FColor VertexColor = CurrentRenderState.GetCurrentVertexColor().ToFColor(true);
+	const FColor LineColor = CurrentRenderState.GetCurrentLineColor().ToFColor(true);
 	const double CurrentLineThickness = GetCurrentLineThickness(View->IsPerspectiveProjection(), View->FOV);
 
 	FPrimitiveDrawInterface* PDI = RenderAPI->GetPrimitiveDrawInterface();
diff --git a/Engine/Source/Runtime/InteractiveToolsFramework/Private/BaseGizmos/GizmoElementLineStrip.cpp b/Engine/Source/Runtime/InteractiveToolsFramework/Private/BaseGizmos/GizmoElementLineStrip.cpp
--- a/Engine/Source/Runtime/InteractiveToolsFramework/Private/BaseGizmos/GizmoElementLineStrip.cpp
+++ b/Engine/Source/Runtime/InteractiveToolsFramework/Private/BaseGizmos/GizmoElementLineStrip.cpp
@@ -33,7 +33,7 @@ void UGizmoElementLineStrip::Render(IToolsContextRenderAPI* RenderAPI, const FRe
 		}
 
 		const float CurrentLineThickness = GetCurrentLineThickness(View->IsPerspectiveProjection(), View->FOV);
-		const FColor LineColor = CurrentRenderState.GetCurrentLineColor().ToFColor(false);
+		const FColor LineColor = CurrentRenderState.GetCurrentLineColor().ToFColor(true);
 
 		// Boost depth bias for lines to avoid z-fighting with the fill (of this or other elements)
 		const float LineDepthBias = 1.0f + CurrentLineThickness;
diff --git a/Engine/Source/Runtime/InteractiveToolsFramework/Private/BaseGizmos/GizmoElementRectangle.cpp b/Engine/Source/Runtime/InteractiveToolsFramework/Private/BaseGizmos/GizmoElementRectangle.cpp
--- a/Engine/Source/Runtime/InteractiveToolsFramework/Private/BaseGizmos/GizmoElementRectangle.cpp
+++ b/Engine/Source/Runtime/InteractiveToolsFramework/Private/BaseGizmos/GizmoElementRectangle.cpp
@@ -34,7 +34,7 @@ void UGizmoElementRectangle::Render(IToolsContextRenderAPI* RenderAPI, const FRe
 		{
 			if (const UMaterialInterface* UseMaterial = CurrentRenderState.GetCurrentMaterial())
 			{
-				FColor VertexColor = CurrentRenderState.GetCurrentVertexColor().ToFColor(false);
+				FColor VertexColor = CurrentRenderState.GetCurrentVertexColor().ToFColor(true);
 				DrawRectangleMesh(PDI, WorldCenter, WorldUpAxis, WorldSideAxis, VertexColor, WorldWidth, WorldHeight, UseMaterial->GetRenderProxy(), SDPG_Foreground);
 			}
 		}
@@ -52,7 +52,7 @@ void UGizmoElementRectangle::Render(IToolsContextRenderAPI* RenderAPI, const FRe
 			// Boost depth bias for lines to avoid z-fighting with the fill (of this or other elements)
 			const float LineDepthBias = 1.0f + CurrentLineThickness;
 
-			FColor LineColor = CurrentRenderState.GetCurrentLineColor().ToFColor(false);
+			FColor LineColor = CurrentRenderState.GetCurrentLineColor().ToFColor(true);
 			DrawRectangle(PDI, WorldCenter, WorldUpAxis, WorldSideAxis, LineColor, WorldWidth, WorldHeight, SDPG_Foreground, CurrentLineThickness, LineDepthBias, bScreenSpaceLine);
 		}
 	}
