[ITF] Added ViewDepthOffset to GizmoElements, which shifts the geometry back and forth in view space to avoid overlaps.

diff --git a/Engine/Source/Runtime/InteractiveToolsFramework/Private/BaseGizmos/GizmoElementBase.cpp b/Engine/Source/Runtime/InteractiveToolsFramework/Private/BaseGizmos/GizmoElementBase.cpp
--- a/Engine/Source/Runtime/InteractiveToolsFramework/Private/BaseGizmos/GizmoElementBase.cpp
+++ b/Engine/Source/Runtime/InteractiveToolsFramework/Private/BaseGizmos/GizmoElementBase.cpp
@@ -154,6 +154,43 @@ bool UGizmoElementBase::GetViewAlignRot(const FVector& InViewLocation, const FVe
 	return true;
 }
 
+void UGizmoElementBase::ApplyViewDepthOffset(const FVector& InViewLocation, const FVector& InViewDirection, FTransform& InOutTransform) const
+{
+	ApplyViewDepthOffset(InViewLocation, InViewDirection, ViewDepthOffset, InOutTransform);
+}
+
+void UGizmoElementBase::ApplyViewDepthOffset(const FVector& InViewLocation, const FVector& InViewDirection, const double& InViewDepthOffset, FTransform& InOutTransform) const
+{
+	const FVector ObjectLocation = InOutTransform.GetLocation();
+
+	FVector CameraToObject = ObjectLocation - InViewLocation;
+	const double DistanceToObject = CameraToObject.Length();
+
+	if (DistanceToObject < UE_KINDA_SMALL_NUMBER)
+	{
+		// Object is at or behind the camera - move along the forward vector without accounting for scale
+		InOutTransform.SetLocation(InViewLocation + InViewDirection * InViewDepthOffset);
+		return;
+	}
+
+	double TargetDistance = DistanceToObject + InViewDepthOffset;
+	if (TargetDistance <= UE_KINDA_SMALL_NUMBER)
+	{
+		// Prevent negative/zero
+		TargetDistance = UE_KINDA_SMALL_NUMBER;
+	}
+
+	// Calculate scale factor to maintain constant visual size
+	const double ScaleFactor = TargetDistance / DistanceToObject;
+
+	// Calculate the new position that maintains the same apparent position in view
+	CameraToObject.Normalize(); // Normalize the vector to get direction
+
+	// Update the transform with new location and scale
+	InOutTransform.SetLocation(InViewLocation + CameraToObject * TargetDistance);
+	InOutTransform.SetScale3D(InOutTransform.GetScale3D() * ScaleFactor);
+}
+
 bool UGizmoElementBase::VerifyUniformScale(const FVector& Scale) const
 {
 	if (!FMath::IsNearlyEqual(Scale.X, Scale.Y, KINDA_SMALL_NUMBER) || !FMath::IsNearlyEqual(Scale.X, Scale.Z, KINDA_SMALL_NUMBER))
@@ -278,6 +315,12 @@ bool UGizmoElementBase::UpdateRenderState(IToolsContextRenderAPI* RenderAPI, con
 	{
 		bOutHasAlignRot = GetViewAlignRot(View, InOutRenderState.LocalToWorldTransform, InLocalCenter, OutAlignRot);
 		InOutRenderState.LocalToWorldTransform = FTransform(OutAlignRot, InLocalCenter) * InOutRenderState.LocalToWorldTransform;
+
+		if (!FMath::IsNearlyZero(ViewDepthOffset))
+		{
+			ApplyViewDepthOffset(View->ViewLocation, View->GetViewDirection(), InOutRenderState.LocalToWorldTransform);
+		}
+	
 		return true;
 	}
 
@@ -602,6 +645,16 @@ float UGizmoElementBase::GetViewAlignAxialAngleTol() const
 	return ViewAlignAxialAngleTol;
 }
 
+void UGizmoElementBase::SetViewDepthOffset(float InViewDepthOffset)
+{
+	ViewDepthOffset = InViewDepthOffset;
+}
+
+float UGizmoElementBase::GetViewDepthOffset() const
+{
+	return ViewDepthOffset;
+}
+
 void UGizmoElementBase::SetPixelHitDistanceThreshold(float InPixelHitDistanceThreshold)
 {
 	PixelHitDistanceThreshold = InPixelHitDistanceThreshold;
diff --git a/Engine/Source/Runtime/InteractiveToolsFramework/Public/BaseGizmos/GizmoElementBase.h b/Engine/Source/Runtime/InteractiveToolsFramework/Public/BaseGizmos/GizmoElementBase.h
--- a/Engine/Source/Runtime/InteractiveToolsFramework/Public/BaseGizmos/GizmoElementBase.h
+++ b/Engine/Source/Runtime/InteractiveToolsFramework/Public/BaseGizmos/GizmoElementBase.h
@@ -254,6 +254,12 @@ public:
 	INTERACTIVETOOLSFRAMEWORK_API virtual void SetViewAlignAxialAngleTol(float InMaxAngleTol);
 	INTERACTIVETOOLSFRAMEWORK_API virtual float GetViewAlignAxialAngleTol() const;
 
+	// Additional view-space offset to apply to the element, affecting both visuals and the hit geometry. Greater values are further away from the camera.
+	// An example use case: A large sphere and view-facing rectangle are both at the origin, but the rectangle should draw on top of the sphere.
+	// The sphere would otherwise overlap the rectangle.
+	INTERACTIVETOOLSFRAMEWORK_API virtual void SetViewDepthOffset(float InViewDepthOffset);
+	INTERACTIVETOOLSFRAMEWORK_API virtual float GetViewDepthOffset() const;
+
 	// Pixel hit distance threshold, element will be scaled enough to add this threshold when line-tracing. */
 	INTERACTIVETOOLSFRAMEWORK_API virtual void SetPixelHitDistanceThreshold(float InPixelHitDistanceThreshold);
 	INTERACTIVETOOLSFRAMEWORK_API virtual float GetPixelHitDistanceThreshold() const;
@@ -528,6 +534,12 @@ protected:
 	UPROPERTY()
 	float PixelHitDistanceThreshold = 7.0;
 
+private:
+
+	// Additional view-space offset to apply to the element, affecting both visuals and the hit geometry.
+	UPROPERTY()
+	float ViewDepthOffset = 0.0;
+
 protected:
 
 	// Return whether element is currently visible.
@@ -564,6 +576,10 @@ protected:
 	// @param OutAlignRot the rotation to align this element in local space, should be prepended to the local-to-world transform.
 	INTERACTIVETOOLSFRAMEWORK_API virtual bool GetViewAlignRot(const FVector& InViewLocation, const FVector& InViewDirection, const FVector& InViewUp, bool bInPerspectiveView, const FTransform& InLocalToWorldTransform, const FVector& InLocalCenter, FQuat& OutAlignRot) const;
 
+	INTERACTIVETOOLSFRAMEWORK_API virtual void ApplyViewDepthOffset(const FVector& InViewLocation, const FVector& InViewDirection, FTransform& InOutTransform) const;
+	
+	INTERACTIVETOOLSFRAMEWORK_API virtual void ApplyViewDepthOffset(const FVector& InViewLocation, const FVector& InViewDirection, const double& InViewDepthOffset, FTransform& InOutTransform) const;
+
 	// Get mesh render state material attribute's value. 
 	INTERACTIVETOOLSFRAMEWORK_API UMaterialInterface* GetMaterial();
 
