[ITF] Added per-state radius and line thickness multipliers to GizmoElements

diff --git a/Engine/Source/Runtime/InteractiveToolsFramework/Private/BaseGizmos/GizmoElementCircle.cpp b/Engine/Source/Runtime/InteractiveToolsFramework/Private/BaseGizmos/GizmoElementCircle.cpp
--- a/Engine/Source/Runtime/InteractiveToolsFramework/Private/BaseGizmos/GizmoElementCircle.cpp
+++ b/Engine/Source/Runtime/InteractiveToolsFramework/Private/BaseGizmos/GizmoElementCircle.cpp
@@ -82,7 +82,7 @@ void UGizmoElementCircle::Render(IToolsContextRenderAPI* RenderAPI, const FRende
 
 	if (bVisibleViewDependent)
 	{
-		const double WorldRadius = Radius * CurrentRenderState.LocalToWorldTransform.GetScale3D().X;
+		const double WorldRadius = GetCurrentRadius() * CurrentRenderState.LocalToWorldTransform.GetScale3D().X;
 		const FVector WorldCenter = CurrentRenderState.LocalToWorldTransform.TransformPosition(FVector::ZeroVector);
 		const FVector WorldAxis0 = CurrentRenderState.LocalToWorldTransform.TransformVectorNoScale(Axis0);
 		const FVector WorldAxis1 = CurrentRenderState.LocalToWorldTransform.TransformVectorNoScale(Axis1);
diff --git a/Engine/Source/Runtime/InteractiveToolsFramework/Private/BaseGizmos/GizmoElementCircleBase.cpp b/Engine/Source/Runtime/InteractiveToolsFramework/Private/BaseGizmos/GizmoElementCircleBase.cpp
--- a/Engine/Source/Runtime/InteractiveToolsFramework/Private/BaseGizmos/GizmoElementCircleBase.cpp
+++ b/Engine/Source/Runtime/InteractiveToolsFramework/Private/BaseGizmos/GizmoElementCircleBase.cpp
@@ -1,15 +1,21 @@
 // Copyright Epic Games, Inc. All Rights Reserved.
  
 #include "BaseGizmos/GizmoElementCircleBase.h"
-#include "BaseGizmos/GizmoRenderingUtil.h"
-#include "BaseGizmos/GizmoMath.h"
-#include "InputState.h"
 #include "Materials/MaterialInterface.h"
 #include "PrimitiveDrawingUtils.h"
-#include "VectorUtil.h"
 
 #include UE_INLINE_GENERATED_CPP_BY_NAME(GizmoElementCircleBase)
 
+double UGizmoElementCircleBase::GetCurrentRadius() const
+{
+	auto GetMultipliedRadius = [this](const float InMultiplier) -> float
+	{
+		return (Radius > UE_SMALL_NUMBER ? Radius * InMultiplier : InMultiplier);
+	};
+
+	return GetMultipliedRadius(RadiusMultiplier.GetValueForState(ElementInteractionState));
+}
+
 bool UGizmoElementCircleBase::IsPartial(const FSceneView* View, const FVector& InWorldCenter, const FVector& InWorldNormal)
 {
 	return IsPartial(InWorldCenter, InWorldNormal, View->ViewLocation, View->GetViewDirection(),
@@ -90,6 +96,46 @@ double UGizmoElementCircleBase::GetRadius() const
 	return Radius;
 }
 
+void UGizmoElementCircleBase::SetHoverRadiusMultiplier(double InHoverRadiusMultiplier)
+{
+	RadiusMultiplier.Hover = InHoverRadiusMultiplier;
+}
+
+double UGizmoElementCircleBase::GetHoverRadiusMultiplier() const
+{
+	return RadiusMultiplier.GetHoverValue();
+}
+
+void UGizmoElementCircleBase::SetInteractRadiusMultiplier(double InInteractRadiusMultiplier)
+{
+	RadiusMultiplier.Interact = InInteractRadiusMultiplier;
+}
+
+double UGizmoElementCircleBase::GetInteractRadiusMultiplier() const
+{
+	return RadiusMultiplier.GetInteractValue();
+}
+
+void UGizmoElementCircleBase::SetSelectRadiusMultiplier(double InSelectRadiusMultiplier)
+{
+	RadiusMultiplier.Select = InSelectRadiusMultiplier;
+}
+
+double UGizmoElementCircleBase::GetSelectRadiusMultiplier() const
+{
+	return RadiusMultiplier.GetSelectValue();
+}
+
+void UGizmoElementCircleBase::SetSubdueRadiusMultiplier(double InSubdueRadiusMultiplier)
+{
+	RadiusMultiplier.Subdue = InSubdueRadiusMultiplier;
+}
+
+double UGizmoElementCircleBase::GetSubdueRadiusMultiplier() const
+{
+	return RadiusMultiplier.GetSubdueValue();
+}
+
 void UGizmoElementCircleBase::SetNumSegments(int32 InNumSegments)
 {
 	NumSegments = InNumSegments;
diff --git a/Engine/Source/Runtime/InteractiveToolsFramework/Private/BaseGizmos/GizmoElementShared.cpp b/Engine/Source/Runtime/InteractiveToolsFramework/Private/BaseGizmos/GizmoElementShared.cpp
new file mode 100644
--- /dev/null
+++ b/Engine/Source/Runtime/InteractiveToolsFramework/Private/BaseGizmos/GizmoElementShared.cpp
@@ -0,0 +1,43 @@
+﻿// Copyright Epic Games, Inc. All Rights Reserved.
+
+#include "BaseGizmos/GizmoElementShared.h"
+
+namespace UE::InteractiveToolsFramework::Private
+{
+	namespace GizmoElementSharedLocals
+	{
+		/** Shared functionality to reduce boilerplate! */
+		template <typename StructType, typename ValueType>
+		static ValueType GetValueForState(const StructType& InPerStateValue, const EGizmoElementInteractionState InState)
+		{
+			switch (InState)
+			{
+			case EGizmoElementInteractionState::Hovering:
+				return InPerStateValue.GetHoverValue();
+
+			case EGizmoElementInteractionState::Interacting:
+				return InPerStateValue.GetInteractValue();
+
+			case EGizmoElementInteractionState::Selected:
+				return InPerStateValue.GetSelectValue();
+
+			case EGizmoElementInteractionState::Subdued:
+				return InPerStateValue.GetSubdueValue();
+
+			case EGizmoElementInteractionState::None:
+			default:
+				return InPerStateValue.GetDefaultValue();
+			}
+		}
+	}
+}
+
+double FGizmoPerStateValueDouble::GetValueForState(const EGizmoElementInteractionState InState) const
+{
+	return UE::InteractiveToolsFramework::Private::GizmoElementSharedLocals::GetValueForState<FGizmoPerStateValueDouble, double>(*this, InState);
+}
+
+FLinearColor FGizmoPerStateValueLinearColor::GetValueForState(const EGizmoElementInteractionState InState) const
+{
+	return UE::InteractiveToolsFramework::Private::GizmoElementSharedLocals::GetValueForState<FGizmoPerStateValueLinearColor, FLinearColor>(*this, InState);
+}
diff --git a/Engine/Source/Runtime/InteractiveToolsFramework/Private/BaseGizmos/GizmoElementTorus.cpp b/Engine/Source/Runtime/InteractiveToolsFramework/Private/BaseGizmos/GizmoElementTorus.cpp
--- a/Engine/Source/Runtime/InteractiveToolsFramework/Private/BaseGizmos/GizmoElementTorus.cpp
+++ b/Engine/Source/Runtime/InteractiveToolsFramework/Private/BaseGizmos/GizmoElementTorus.cpp
@@ -21,7 +21,6 @@ void UGizmoElementTorus::Render(IToolsContextRenderAPI* RenderAPI, const FRender
 		if (const UMaterialInterface* UseMaterial = CurrentRenderState.GetCurrentMaterial())
 		{
 			const FVector WorldCenter = CurrentRenderState.LocalToWorldTransform.GetLocation();
-			const FVector WorldAxis0 = CurrentRenderState.LocalToWorldTransform.TransformVectorNoScale(Axis0);
 			const FVector Normal = Axis0 ^ Axis1; 
 			const FVector WorldNormal = CurrentRenderState.LocalToWorldTransform.TransformVectorNoScale(Normal);
 
@@ -39,8 +38,10 @@ void UGizmoElementTorus::Render(IToolsContextRenderAPI* RenderAPI, const FRender
 			FVector TorusSideAxis = Normal ^ BeginAxis;
 			TorusSideAxis.Normalize();
 
-			DrawTorus(PDI, CurrentRenderState.LocalToWorldTransform.ToMatrixWithScale(), 
-				BeginAxis, TorusSideAxis, Radius, InnerRadius, NumSegments, NumInnerSlices,  
+			const FColor Color = CurrentRenderState.GetCurrentVertexColor().ToFColor(true);
+
+			DrawTorus(PDI, CurrentRenderState.LocalToWorldTransform.ToMatrixWithScale(),
+				BeginAxis, TorusSideAxis, Color, GetCurrentRadius(), GetCurrentInnerRadius(), NumSegments, NumInnerSlices,  
 				UseMaterial->GetRenderProxy(), SDPG_Foreground, bPartial, PartialAngle, bEndCaps);
 		}
 	}
@@ -62,8 +63,8 @@ FInputRayHit UGizmoElementTorus::LineTrace(const UGizmoViewContext* ViewContext,
 	if (bHittableViewDependent)
 	{
 		const double PixelHitThresholdAdjust = CurrentLineTraceState.PixelToWorldScale * PixelHitDistanceThreshold;
-		const double WorldOuterRadius = Radius * CurrentLineTraceState.LocalToWorldTransform.GetScale3D().X;
-		const double WorldInnerRadius = InnerRadius * CurrentLineTraceState.LocalToWorldTransform.GetScale3D().X;
+		const double WorldOuterRadius = GetCurrentRadius() * CurrentLineTraceState.LocalToWorldTransform.GetScale3D().X;
+		const double WorldInnerRadius = GetCurrentInnerRadius() * CurrentLineTraceState.LocalToWorldTransform.GetScale3D().X;
 		const FVector WorldCenter = CurrentLineTraceState.LocalToWorldTransform.GetLocation();
 		const FVector WorldAxis0 = CurrentLineTraceState.LocalToWorldTransform.TransformVectorNoScale(Axis0);
 		const FVector WorldAxis1 = CurrentLineTraceState.LocalToWorldTransform.TransformVectorNoScale(Axis1);
@@ -219,6 +220,16 @@ FInputRayHit UGizmoElementTorus::LineTrace(const UGizmoViewContext* ViewContext,
 	return FInputRayHit();
 }
 
+double UGizmoElementTorus::GetCurrentInnerRadius() const
+{
+	auto GetMultipliedInnerRadius = [this](const float InMultiplier) -> float
+	{
+		return (InnerRadius > UE_SMALL_NUMBER ? InnerRadius * InMultiplier : InMultiplier);
+	};
+
+	return GetMultipliedInnerRadius(InnerRadiusMultiplier.GetValueForState(ElementInteractionState));
+}
+
 bool UGizmoElementTorus::IsGlancingAngle(const UGizmoViewContext* View, const FVector& InWorldCenter, const FVector& InWorldNormal)
 {
 	check(View);
@@ -252,6 +263,46 @@ double UGizmoElementTorus::GetInnerRadius() const
 	return InnerRadius;
 }
 
+void UGizmoElementTorus::SetHoverInnerRadiusMultiplier(double InHoverInnerRadiusMultiplier)
+{
+	InnerRadiusMultiplier.Hover = InHoverInnerRadiusMultiplier;
+}
+
+double UGizmoElementTorus::GetHoverInnerRadiusMultiplier() const
+{
+	return InnerRadiusMultiplier.GetHoverValue();
+}
+
+void UGizmoElementTorus::SetInteractInnerRadiusMultiplier(double InInteractInnerRadiusMultiplier)
+{
+	InnerRadiusMultiplier.Interact = InInteractInnerRadiusMultiplier;
+}
+
+double UGizmoElementTorus::GetInteractInnerRadiusMultiplier() const
+{
+	return InnerRadiusMultiplier.GetInteractValue();
+}
+
+void UGizmoElementTorus::SetSelectInnerRadiusMultiplier(double InSelectInnerRadiusMultiplier)
+{
+	InnerRadiusMultiplier.Select = InSelectInnerRadiusMultiplier;
+}
+
+double UGizmoElementTorus::GetSelectInnerRadiusMultiplier() const
+{
+	return InnerRadiusMultiplier.GetSelectValue();
+}
+
+void UGizmoElementTorus::SetSubdueInnerRadiusMultiplier(double InSubdueInnerRadiusMultiplier)
+{
+	InnerRadiusMultiplier.Subdue = InSubdueInnerRadiusMultiplier;
+}
+
+double UGizmoElementTorus::GetSubdueInnerRadiusMultiplier() const
+{
+	return InnerRadiusMultiplier.GetSubdueValue();
+}
+
 void UGizmoElementTorus::SetNumInnerSlices(int32 InNumInnerSlices)
 {
 	NumInnerSlices = InNumInnerSlices;
diff --git a/Engine/Source/Runtime/InteractiveToolsFramework/Public/BaseGizmos/GizmoElementCircleBase.h b/Engine/Source/Runtime/InteractiveToolsFramework/Public/BaseGizmos/GizmoElementCircleBase.h
--- a/Engine/Source/Runtime/InteractiveToolsFramework/Public/BaseGizmos/GizmoElementCircleBase.h
+++ b/Engine/Source/Runtime/InteractiveToolsFramework/Public/BaseGizmos/GizmoElementCircleBase.h
@@ -19,6 +19,8 @@ class UGizmoElementCircleBase : public UGizmoElementLineBase
 {
 	GENERATED_BODY()
 
+	friend class FGizmoElementAccessor;
+
 public:
 
 	// CircleBase center.
@@ -38,6 +40,22 @@ public:
 	INTERACTIVETOOLSFRAMEWORK_API virtual void SetRadius(double InRadius);
 	INTERACTIVETOOLSFRAMEWORK_API virtual double GetRadius() const;
 
+	// Multiplier applied to the circle radius when hovering
+	INTERACTIVETOOLSFRAMEWORK_API virtual void SetHoverRadiusMultiplier(double InHoverRadiusMultiplier);
+	INTERACTIVETOOLSFRAMEWORK_API virtual double GetHoverRadiusMultiplier() const;
+
+	// Multiplier applied to the circle radius when interacting
+	INTERACTIVETOOLSFRAMEWORK_API virtual void SetInteractRadiusMultiplier(double InInteractRadiusMultiplier);
+	INTERACTIVETOOLSFRAMEWORK_API virtual double GetInteractRadiusMultiplier() const;
+
+	// Multiplier applied to the circle radius when selected
+	INTERACTIVETOOLSFRAMEWORK_API virtual void SetSelectRadiusMultiplier(double InSelectRadiusMultiplier);
+	INTERACTIVETOOLSFRAMEWORK_API virtual double GetSelectRadiusMultiplier() const;
+
+	// Multiplier applied to the circle radius when subdued
+	INTERACTIVETOOLSFRAMEWORK_API virtual void SetSubdueRadiusMultiplier(double InSubdueRadiusMultiplier);
+	INTERACTIVETOOLSFRAMEWORK_API virtual double GetSubdueRadiusMultiplier() const;
+
 	// Number of segments for rendering circle.
 	INTERACTIVETOOLSFRAMEWORK_API virtual void SetNumSegments(int32 InNumSegments);
 	INTERACTIVETOOLSFRAMEWORK_API virtual int32 GetNumSegments() const;
@@ -79,6 +97,10 @@ protected:
 	UPROPERTY()
 	double Radius = 100.0;
 
+	// Optional radius multipliers to apply when hovering, interacting, selected or subdued.
+	UPROPERTY()
+	FGizmoPerStateValueDouble RadiusMultiplier;
+
 	// Number of segments for rendering arc.
 	UPROPERTY()
 	int32 NumSegments = 64;
@@ -102,6 +124,9 @@ protected:
 
 protected:
 
+	// Returns the current radius based on the current interaction state.
+	INTERACTIVETOOLSFRAMEWORK_API double GetCurrentRadius() const;
+
 	// Returns whether element should be partial based on current view
 	INTERACTIVETOOLSFRAMEWORK_API bool IsPartial(const FSceneView* View, const FVector& InWorldCenter, const FVector& InWorldNormal);
 
@@ -112,4 +137,5 @@ protected:
 	INTERACTIVETOOLSFRAMEWORK_API bool IsPartial(const FVector& InWorldCenter, const FVector& InWorldNormal,
 		const FVector& InViewLocation, const FVector& InViewDirection,
 		const bool bIsPerspectiveProjection);
+
 };
diff --git a/Engine/Source/Runtime/InteractiveToolsFramework/Public/BaseGizmos/GizmoElementShared.h b/Engine/Source/Runtime/InteractiveToolsFramework/Public/BaseGizmos/GizmoElementShared.h
--- a/Engine/Source/Runtime/InteractiveToolsFramework/Public/BaseGizmos/GizmoElementShared.h
+++ b/Engine/Source/Runtime/InteractiveToolsFramework/Public/BaseGizmos/GizmoElementShared.h
@@ -94,3 +94,108 @@ enum class EGizmoElementDrawType : uint8
 	FillAndLine = Fill | Line
 };
 ENUM_CLASS_FLAGS(EGizmoElementDrawType)
+
+/**
+ * Used to store per-state (double) values for gizmo elements.
+ * ie. line thickness multipliers.
+ */
+USTRUCT(MinimalAPI, meta = (DisplayName = "Per-State Value (Double)"))
+struct FGizmoPerStateValueDouble
+{
+	GENERATED_BODY()
+
+	/**
+	 * Default value, used when the Interaction State is "None".
+	 * Optional to allow explicit un-setting, implying inheritance or some other value source.
+	 */
+	UPROPERTY(EditAnywhere, Category = "Value")
+	TOptional<double> Default;
+
+	/** Value used when hovering. */
+	UPROPERTY(EditAnywhere, Category = "Value")
+	TOptional<double> Hover;
+
+	/** Value used when interacting. */
+	UPROPERTY(EditAnywhere, Category = "Value")
+	TOptional<double> Interact;
+
+	/** Value used when selected. */
+	UPROPERTY(EditAnywhere, Category = "Value")
+	TOptional<double> Select;
+
+	/** Value used when subdued. */
+	UPROPERTY(EditAnywhere, Category = "Value")
+	TOptional<double> Subdue;
+
+	/** Get the value for the given Interaction State, or Default if not set. */
+	INTERACTIVETOOLSFRAMEWORK_API double GetValueForState(const EGizmoElementInteractionState InState) const;
+
+	/** Get the Default value if set, otherwise "ValueDefault" */
+	double GetDefaultValue() const { return Default.Get(ValueDefault); }
+
+	/** Get the Hover value, or Default if not set. */
+	double GetHoverValue() const { return Hover.Get(GetDefaultValue()); }
+
+	/** Get the Interact value, or Default if not set. */
+	double GetInteractValue() const { return Interact.Get(GetDefaultValue()); }
+
+	/** Get the Select value, or Default if not set. */
+	double GetSelectValue() const { return Select.Get(GetDefaultValue()); }
+
+	/** Get the Subdue value, or Default if not set. */
+	double GetSubdueValue() const { return Subdue.Get(GetDefaultValue()); }
+
+private:
+	static constexpr double ValueDefault = 1.0;
+};
+
+/**
+ * Used to store per-state (Linear Color) values for gizmo elements.
+ * ie. vertex color.
+ */
+USTRUCT(MinimalAPI, meta = (DisplayName = "Per-State Value (LinearColor)"))
+struct FGizmoPerStateValueLinearColor
+{
+	GENERATED_BODY()
+
+	/**
+	 * Default value, used when the Interaction State is "None".
+	 * Optional to allow explicit un-setting, implying inheritance or some other value source.
+	 */
+	UPROPERTY(EditAnywhere, Category = "Value")
+	TOptional<FLinearColor> Default;
+
+	/** Value used when hovering. */
+	UPROPERTY(EditAnywhere, Category = "Value")
+	TOptional<FLinearColor> Hover;
+
+	/** Value used when interacting. */
+	UPROPERTY(EditAnywhere, Category = "Value")
+	TOptional<FLinearColor> Interact;
+
+	/** Value used when selected. */
+	UPROPERTY(EditAnywhere, Category = "Value")
+	TOptional<FLinearColor> Select;
+
+	/** Value used when subdued. */
+	UPROPERTY(EditAnywhere, Category = "Value")
+	TOptional<FLinearColor> Subdue;
+
+	/** Get the value for the given Interaction State, or Default if not set. */
+	INTERACTIVETOOLSFRAMEWORK_API FLinearColor GetValueForState(const EGizmoElementInteractionState InState) const;
+
+	/** Get the Default value if set, otherwise FLinearColor::Transparent. */
+	FLinearColor GetDefaultValue() const { return Default.Get(FLinearColor::Transparent); }
+
+	/** Get the Hover value, or Default if not set. */
+	FLinearColor GetHoverValue() const { return Hover.Get(GetDefaultValue()); }
+
+	/** Get the Interact value, or Default if not set. */
+	FLinearColor GetInteractValue() const { return Interact.Get(GetDefaultValue()); }
+
+	/** Get the Select value, or Default if not set. */
+	FLinearColor GetSelectValue() const { return Select.Get(GetDefaultValue()); }
+
+	/** Get the Subdue value, or Default if not set. */
+	FLinearColor GetSubdueValue() const { return Subdue.Get(GetDefaultValue()); }
+};
diff --git a/Engine/Source/Runtime/InteractiveToolsFramework/Public/BaseGizmos/GizmoElementTorus.h b/Engine/Source/Runtime/InteractiveToolsFramework/Public/BaseGizmos/GizmoElementTorus.h
--- a/Engine/Source/Runtime/InteractiveToolsFramework/Public/BaseGizmos/GizmoElementTorus.h
+++ b/Engine/Source/Runtime/InteractiveToolsFramework/Public/BaseGizmos/GizmoElementTorus.h
@@ -39,6 +39,22 @@ public:
 	INTERACTIVETOOLSFRAMEWORK_API virtual void SetInnerRadius(double InInnerRadius);
 	INTERACTIVETOOLSFRAMEWORK_API virtual double GetInnerRadius() const;
 
+	// Multiplier applied to the inner radius when hovering
+	INTERACTIVETOOLSFRAMEWORK_API virtual void SetHoverInnerRadiusMultiplier(double InHoverInnerRadiusMultiplier);
+	INTERACTIVETOOLSFRAMEWORK_API virtual double GetHoverInnerRadiusMultiplier() const;
+
+	// Multiplier applied to the inner radius when interacting
+	INTERACTIVETOOLSFRAMEWORK_API virtual void SetInteractInnerRadiusMultiplier(double InInteractInnerRadiusMultiplier);
+	INTERACTIVETOOLSFRAMEWORK_API virtual double GetInteractInnerRadiusMultiplier() const;
+
+	// Multiplier applied to the inner radius when selected
+	INTERACTIVETOOLSFRAMEWORK_API virtual void SetSelectInnerRadiusMultiplier(double InSelectInnerRadiusMultiplier);
+	INTERACTIVETOOLSFRAMEWORK_API virtual double GetSelectInnerRadiusMultiplier() const;
+
+	// Multiplier applied to the inner radius when subdued
+	INTERACTIVETOOLSFRAMEWORK_API virtual void SetSubdueInnerRadiusMultiplier(double InSubdueInnerRadiusMultiplier);
+	INTERACTIVETOOLSFRAMEWORK_API virtual double GetSubdueInnerRadiusMultiplier() const;
+
 	// Number of inner slices for rendering torus.
 	INTERACTIVETOOLSFRAMEWORK_API virtual void SetNumInnerSlices(int32 InNumInnerSlices);
 	INTERACTIVETOOLSFRAMEWORK_API virtual int32 GetNumInnerSlices() const;
@@ -53,6 +69,10 @@ protected:
 	UPROPERTY()
 	double InnerRadius = 5.0;
 
+	// Optional inner radius multipliers to apply when hovering, interacting, selected or subdued.
+	UPROPERTY()
+	FGizmoPerStateValueDouble InnerRadiusMultiplier;
+
 	// Number of slices to render in each torus segment.
 	UPROPERTY()
 	int32 NumInnerSlices = 8;
@@ -63,6 +83,9 @@ protected:
 
 protected:
 
+	// Returns the current inner radius based on the current interaction state.
+	INTERACTIVETOOLSFRAMEWORK_API double GetCurrentInnerRadius() const;
+
 	INTERACTIVETOOLSFRAMEWORK_API bool IsGlancingAngle(const UGizmoViewContext* View, const FVector& InWorldCenter, const FVector& InWorldNormal);
 
 };
