[ITF] Accounted for ViewDepthOffset in GizmoElementCircle

diff --git a/Engine/Source/Runtime/InteractiveToolsFramework/Private/BaseGizmos/GizmoElementCircle.cpp b/Engine/Source/Runtime/InteractiveToolsFramework/Private/BaseGizmos/GizmoElementCircle.cpp
--- a/Engine/Source/Runtime/InteractiveToolsFramework/Private/BaseGizmos/GizmoElementCircle.cpp
+++ b/Engine/Source/Runtime/InteractiveToolsFramework/Private/BaseGizmos/GizmoElementCircle.cpp
@@ -1,9 +1,8 @@
 // Copyright Epic Games, Inc. All Rights Reserved.
  
 #include "BaseGizmos/GizmoElementCircle.h"
-
-#include "BaseGizmos/GizmoMath.h"
 #include "BaseGizmos/GizmoRenderingUtil.h"
+#include "BaseGizmos/GizmoMath.h"
 #include "DynamicMeshBuilder.h"
 #include "InputState.h"
 #include "Intersection/IntersectionUtil.h"
@@ -65,7 +64,7 @@ namespace UE::InteractiveToolsFramework::Private
 				MeshBuilder.AddTriangle(V0, V2, V1);
 			}
 
-			MeshBuilder.Draw(PDI, FMatrix::Identity, MaterialRenderProxy, DepthPriority, 0.f);
+			MeshBuilder.Draw(PDI, FMatrix::Identity, MaterialRenderProxy, DepthPriority,0.f);
 		}
 	}
 }
@@ -86,15 +85,29 @@ void UGizmoElementCircle::Render(IToolsContextRenderAPI* RenderAPI, const FRende
 		const FVector WorldCenter = CurrentRenderState.LocalToWorldTransform.TransformPosition(FVector::ZeroVector);
 		const FVector WorldAxis0 = CurrentRenderState.LocalToWorldTransform.TransformVectorNoScale(Axis0);
 		const FVector WorldAxis1 = CurrentRenderState.LocalToWorldTransform.TransformVectorNoScale(Axis1);
-
 		FPrimitiveDrawInterface* PDI = RenderAPI->GetPrimitiveDrawInterface();
 
+		const FSceneView* View = RenderAPI->GetSceneView();
+		if (!ensure(View))
+		{
+			return;
+		}
+
+		const float CurrentLineThickness = GetCurrentLineThickness(View->IsPerspectiveProjection(), View->FOV);
+
 		if (bDrawMesh)
 		{
 			if (const UMaterialInterface* UseMaterial = CurrentRenderState.GetCurrentMaterial())
 			{
 				using namespace UE::InteractiveToolsFramework::Private;
 
+				// Offset in view-space to avoid z-fighting with the line
+				if (bDrawLine)
+				{
+					const double MeshViewDepthOffset = GetViewDepthOffset() - CurrentLineThickness;
+					ApplyViewDepthOffset(View->ViewLocation, View->GetViewDirection(), MeshViewDepthOffset, CurrentRenderState.LocalToWorldTransform);
+				}
+
 				const FColor VertexColor = CurrentRenderState.GetCurrentVertexColor().ToFColor(true);
 				GizmoElementCircleLocals::DrawDisc(PDI, WorldCenter, WorldAxis0, WorldAxis1, VertexColor, WorldRadius, NumSegments, UseMaterial->GetRenderProxy(), SDPG_Foreground);
 			}
@@ -102,19 +115,12 @@ void UGizmoElementCircle::Render(IToolsContextRenderAPI* RenderAPI, const FRende
 
 		if (bDrawLine)
 		{
-			const FSceneView* View = RenderAPI->GetSceneView();
-			if (!ensure(View))
-			{
-				return;
-			}
-
-			float CurrentLineThickness = GetCurrentLineThickness(View->IsPerspectiveProjection(), View->FOV);
-
 			// Boost depth bias for lines to avoid z-fighting with the fill (of this or other elements)
 			const float LineDepthBias = 1.0f + CurrentLineThickness;
 
 			DrawCircle(PDI, WorldCenter, WorldAxis0, WorldAxis1, CurrentRenderState.GetCurrentLineColor(), WorldRadius, NumSegments, SDPG_Foreground, CurrentLineThickness, LineDepthBias, bScreenSpaceLine);
 		}
+
 	}
 }
 
