[ITF] FGizmoPerStateValueLinearColor getters return by ref, added equality operators to both *PerStateValue types.

diff --git a/Engine/Source/Runtime/InteractiveToolsFramework/Private/BaseGizmos/GizmoElementShared.cpp b/Engine/Source/Runtime/InteractiveToolsFramework/Private/BaseGizmos/GizmoElementShared.cpp
--- a/Engine/Source/Runtime/InteractiveToolsFramework/Private/BaseGizmos/GizmoElementShared.cpp
+++ b/Engine/Source/Runtime/InteractiveToolsFramework/Private/BaseGizmos/GizmoElementShared.cpp
@@ -37,7 +37,7 @@ double FGizmoPerStateValueDouble::GetValueForState(const EGizmoElementInteractio
 	return UE::InteractiveToolsFramework::Private::GizmoElementSharedLocals::GetValueForState<FGizmoPerStateValueDouble, double>(*this, InState);
 }
 
-FLinearColor FGizmoPerStateValueLinearColor::GetValueForState(const EGizmoElementInteractionState InState) const
+const FLinearColor& FGizmoPerStateValueLinearColor::GetValueForState(const EGizmoElementInteractionState InState) const
 {
-	return UE::InteractiveToolsFramework::Private::GizmoElementSharedLocals::GetValueForState<FGizmoPerStateValueLinearColor, FLinearColor>(*this, InState);
+	return UE::InteractiveToolsFramework::Private::GizmoElementSharedLocals::GetValueForState<FGizmoPerStateValueLinearColor, const FLinearColor&>(*this, InState);
 }
diff --git a/Engine/Source/Runtime/InteractiveToolsFramework/Public/BaseGizmos/GizmoElementShared.h b/Engine/Source/Runtime/InteractiveToolsFramework/Public/BaseGizmos/GizmoElementShared.h
--- a/Engine/Source/Runtime/InteractiveToolsFramework/Public/BaseGizmos/GizmoElementShared.h
+++ b/Engine/Source/Runtime/InteractiveToolsFramework/Public/BaseGizmos/GizmoElementShared.h
@@ -145,6 +145,20 @@ struct FGizmoPerStateValueDouble
 	/** Get the Subdue value, or Default if not set. */
 	double GetSubdueValue() const { return Subdue.Get(GetDefaultValue()); }
 
+	friend bool operator==(const FGizmoPerStateValueDouble& InLeft, const FGizmoPerStateValueDouble& InRight)
+	{
+		return InLeft.Default == InRight.Default
+			&& InLeft.Hover == InRight.Hover
+			&& InLeft.Interact == InRight.Interact
+			&& InLeft.Select == InRight.Select
+			&& InLeft.Subdue == InRight.Subdue;
+	}
+
+	friend bool operator!=(const FGizmoPerStateValueDouble& InLeft, const FGizmoPerStateValueDouble& InRight)
+	{
+		return !(InLeft == InRight);
+	}
+
 private:
 	static constexpr double ValueDefault = 1.0;
 };
@@ -182,20 +196,34 @@ struct FGizmoPerStateValueLinearColor
 	TOptional<FLinearColor> Subdue;
 
 	/** Get the value for the given Interaction State, or Default if not set. */
-	INTERACTIVETOOLSFRAMEWORK_API FLinearColor GetValueForState(const EGizmoElementInteractionState InState) const;
+	INTERACTIVETOOLSFRAMEWORK_API const FLinearColor& GetValueForState(const EGizmoElementInteractionState InState) const;
 
 	/** Get the Default value if set, otherwise FLinearColor::Transparent. */
-	FLinearColor GetDefaultValue() const { return Default.Get(FLinearColor::Transparent); }
+	const FLinearColor& GetDefaultValue() const { return Default.Get(FLinearColor::Transparent); }
 
 	/** Get the Hover value, or Default if not set. */
-	FLinearColor GetHoverValue() const { return Hover.Get(GetDefaultValue()); }
+	const FLinearColor& GetHoverValue() const { return Hover.Get(GetDefaultValue()); }
 
 	/** Get the Interact value, or Default if not set. */
-	FLinearColor GetInteractValue() const { return Interact.Get(GetDefaultValue()); }
+	const FLinearColor& GetInteractValue() const { return Interact.Get(GetDefaultValue()); }
 
 	/** Get the Select value, or Default if not set. */
-	FLinearColor GetSelectValue() const { return Select.Get(GetDefaultValue()); }
+	const FLinearColor& GetSelectValue() const { return Select.Get(GetDefaultValue()); }
 
 	/** Get the Subdue value, or Default if not set. */
-	FLinearColor GetSubdueValue() const { return Subdue.Get(GetDefaultValue()); }
+	const FLinearColor& GetSubdueValue() const { return Subdue.Get(GetDefaultValue()); }
+
+	friend bool operator==(const FGizmoPerStateValueLinearColor& InLeft, const FGizmoPerStateValueLinearColor& InRight)
+	{
+		return InLeft.Default == InRight.Default
+			&& InLeft.Hover == InRight.Hover
+			&& InLeft.Interact == InRight.Interact
+			&& InLeft.Select == InRight.Select
+			&& InLeft.Subdue == InRight.Subdue;
+	}
+
+	friend bool operator!=(const FGizmoPerStateValueLinearColor& Lhs, const FGizmoPerStateValueLinearColor& RHS)
+	{
+		return !(Lhs == RHS);
+	}
 };
