[ITF] Added DrawDirectionalArrow to ToolDataVisualizer

diff --git a/Engine/Source/Runtime/InteractiveToolsFramework/Private/ToolDataVisualizer.cpp b/Engine/Source/Runtime/InteractiveToolsFramework/Private/ToolDataVisualizer.cpp
--- a/Engine/Source/Runtime/InteractiveToolsFramework/Private/ToolDataVisualizer.cpp
+++ b/Engine/Source/Runtime/InteractiveToolsFramework/Private/ToolDataVisualizer.cpp
@@ -92,6 +92,59 @@ void FToolDataVisualizer::InternalDrawTransformedLine(const FVector& A, const FV
 		LineThicknessIn * PDISizeScale, DepthBias, true);
 }
 
+void FToolDataVisualizer::InternalDrawDirectionalArrow(const FVector& A, const FVector& B, const FVector& InNormal, const FLinearColor& Color, const float InHeadSize, float InLineThickness, bool bInDepthTested)
+{
+	FVector Direction = B - A;
+	Direction.Normalize();
+
+	FVector Up = FVector::UpVector;
+	FVector Right = Direction ^ Up;
+
+	if (!Right.IsNormalized())
+	{
+		Direction.FindBestAxisVectors(Up, Right);
+	}
+
+	FMatrix Transform = FMatrix::Identity;
+	Transform.SetAxes(&Direction, &Right, &Up, &B);
+	
+	CurrentPDI->DrawLine(
+		A,
+		B,
+		Color,
+		uint8( (bInDepthTested) ? SDPG_World : SDPG_Foreground),
+		InLineThickness * PDISizeScale,
+		DepthBias,
+		true);
+
+	const float ArrowSqrt = FMath::Sqrt(InHeadSize);
+
+	const FVector ArrowPos0 = FVector(-ArrowSqrt, +ArrowSqrt, +ArrowSqrt);
+	const FVector ArrowPos1 = FVector(-ArrowSqrt, +ArrowSqrt, -ArrowSqrt);
+	const FVector ArrowPos2 = FVector(-ArrowSqrt, -ArrowSqrt, +ArrowSqrt);
+	const FVector ArrowPos3 = FVector(-ArrowSqrt, -ArrowSqrt, -ArrowSqrt);
+
+	auto DrawHeadLine = [&](const FVector& InB)
+	{
+		const FVector HeadA = Transform.TransformPosition(FVector::ZeroVector);
+		const FVector HeadB = Transform.TransformPosition(InB);
+		
+		CurrentPDI->DrawLine(
+			HeadA,
+			HeadB,
+			Color,
+			uint8( (bInDepthTested) ? SDPG_World : SDPG_Foreground),
+			InLineThickness * PDISizeScale,
+			DepthBias,
+			true);
+	};
+
+	DrawHeadLine(ArrowPos0);
+	DrawHeadLine(ArrowPos1);
+	DrawHeadLine(ArrowPos2);
+	DrawHeadLine(ArrowPos3);
+}
+
 
 void FToolDataVisualizer::InternalDrawTransformedPoint(const FVector& Position, const FLinearColor& ColorIn, float PointSizeIn, bool bDepthTestedIn)
 {
@@ -169,14 +222,19 @@ void FToolDataVisualizer::InternalDrawSquare(const FVector& Center, const FVecto
 // @note: This is taken from MeshDebugDraw::DrawSimpleGrid
 void FToolDataVisualizer::InternalDrawGrid(const FVector& Center, const FVector& Normal, const FVector& Size, const FLinearColor& Color, float LineThicknessIn, bool bDepthTestedIn)
 {
-	const ESceneDepthPriorityGroup DepthPriority = (bDepthTestedIn) ? SDPG_World : SDPG_Foreground;
-
 	FVector Tan1, Tan2;
 	GizmoMath::MakeNormalPlaneBasis(Normal, Tan1, Tan2);
 
+	InternalDrawGrid(Center, Normal, Tan1, Tan2, Size, Color, LineThicknessIn, bDepthTestedIn);
+}
+
+void FToolDataVisualizer::InternalDrawGrid(const FVector& Center, const FVector& Normal, const FVector& Up, const FVector& Side, const FVector& Size, const FLinearColor& Color, float LineThicknessIn, bool bDepthTestedIn)
+{
+	const ESceneDepthPriorityGroup DepthPriority = (bDepthTestedIn) ? SDPG_World : SDPG_Foreground;
+
 	const FVector CC = TransformP(Center);
-	const FVector SA = TransformV(Tan1);
-	const FVector SB = TransformV(Tan2);
+	const FVector SA = TransformV(Up);
+	const FVector SB = TransformV(Side);
 
 	// Always odd
 	constexpr int32 NumGridLines = (FMath::DivideAndRoundUp(16, 2) * 2) + 1;
diff --git a/Engine/Source/Runtime/InteractiveToolsFramework/Public/ToolDataVisualizer.h b/Engine/Source/Runtime/InteractiveToolsFramework/Public/ToolDataVisualizer.h
--- a/Engine/Source/Runtime/InteractiveToolsFramework/Public/ToolDataVisualizer.h
+++ b/Engine/Source/Runtime/InteractiveToolsFramework/Public/ToolDataVisualizer.h
@@ -22,6 +22,8 @@ public:
 	FLinearColor LineColor;   // = Red
 	/** Default thickness used for drawing lines */
 	float LineThickness = 1.0f;
+	/** Default size used for drawing arrow heads */
+	float ArrowHeadSize = 4.0f;
 
 	/** Default color used for drawing points */
 	FLinearColor PointColor;   // = Red
@@ -100,6 +102,7 @@ public:
 	{
 		LineColor = Color;
 		LineThickness = Thickness;
+		ArrowHeadSize = InArrowHeadSize;
 	}
 
 	/** Update the default point color and size */
@@ -143,6 +146,35 @@ public:
 		InternalDrawTransformedLine(TransformP((FVector)A), TransformP((FVector)B), Color, LineThicknessIn, bDepthTestedIn);
 	}
 
+	/** Draw an arrow with default parameters */
+	template<typename PointType>
+	void DrawDirectionalArrow(const PointType& A, const PointType& B, const FVector& InNormal)
+	{
+		InternalDrawDirectionalArrow(TransformP((FVector)A), TransformP((FVector)B), TransformN(InNormal), LineColor, ArrowHeadSize, LineThickness, bDepthTested);
+	}
+
+	/** Draw a line with the given Color, otherwise use default parameters */
+	template<typename PointType>
+	void DrawDirectionalArrow(const PointType& A, const PointType& B, const FVector& InNormal, const FLinearColor& Color)
+	{
+		InternalDrawDirectionalArrow(TransformP((FVector)A), TransformP((FVector)B), TransformN(InNormal), Color, ArrowHeadSize, LineThickness, bDepthTested);
+	}
+
+	/** Draw a line with the given Color and Thickness, otherwise use default parameters */
+	template<typename PointType>
+	void DrawDirectionalArrow(const PointType& A, const PointType& B, const FVector& InNormal, const FLinearColor& Color, float InArrowHeadSize, float LineThicknessIn)
+	{
+		InternalDrawDirectionalArrow(TransformP((FVector)A), TransformP((FVector)B), TransformN(InNormal), Color, InArrowHeadSize, LineThicknessIn, bDepthTested);
+	}
+
+	/** Draw a line with the given parameters */
+	template<typename PointType>
+	void DrawDirectionalArrow(const PointType& A, const PointType& B, const FVector& InNormal, const FLinearColor& Color, float InArrowHeadSize, float LineThicknessIn, bool bDepthTestedIn)
+	{
+		InternalDrawDirectionalArrow(TransformP((FVector)A), TransformP((FVector)B), TransformN(InNormal), Color, InArrowHeadSize, LineThicknessIn, bDepthTestedIn);
+	}
+
+
 	/** Draw a point with default parameters */
 	template<typename PointType>
 	void DrawPoint(const PointType& Position)
@@ -250,6 +282,12 @@ public:
 		InternalDrawGrid(Center, Normal, Size, Color, LineThicknessIn, bDepthTestedIn);
 	}
 
+	template<typename PointType>
+	void DrawGrid(const PointType& Center, const PointType& Normal,  const PointType& Up, const PointType& Side, const PointType& Size, const FLinearColor& Color, float LineThicknessIn, bool bDepthTestedIn)
+	{
+		InternalDrawGrid(Center, Normal, Up, Side, Size, Color, LineThicknessIn, bDepthTestedIn);
+	}
+
 	/** Draw a 3D disc at given position/normal with the given parameters. Tangent axes are defined internally */
 	template<typename PointType>
 	void DrawDisc(const PointType& Position, const PointType& Normal, float Radius, int Steps, const FColor& Color, FMaterialRenderProxy* RenderProxy, bool bDepthTestedIn)
@@ -274,6 +312,8 @@ protected:
 
 	// actually does the line drawing; assumes A and B are already transformed
 	INTERACTIVETOOLSFRAMEWORK_API virtual void InternalDrawTransformedLine(const FVector& A, const FVector& B, const FLinearColor& Color, float LineThickness, bool bDepthTested);
+	// actually does the arrow drawing; assumes A and B are already transformed
+	INTERACTIVETOOLSFRAMEWORK_API virtual void InternalDrawDirectionalArrow(const FVector& A, const FVector& B, const FVector& InNormal, const FLinearColor& Color, const float InHeadSize, float LineThickness, bool bDepthTested);
 	// actually does the point drawing; assumes Position is already transformed
 	INTERACTIVETOOLSFRAMEWORK_API virtual void InternalDrawTransformedPoint(const FVector& Position, const FLinearColor& Color, float PointSize, bool bDepthTested);
 
@@ -294,6 +334,8 @@ protected:
 	INTERACTIVETOOLSFRAMEWORK_API virtual void InternalDrawSquare(const FVector& Center, const FVector& SideA, const FVector& SideB, const FLinearColor& Color, float LineThickness, bool bDepthTested);
 	// actually does the grid drawing
 	INTERACTIVETOOLSFRAMEWORK_API virtual void InternalDrawGrid(const FVector& Center, const FVector& Normal, const FVector& Size, const FLinearColor& Color, float LineThickness, bool bDepthTested);
+	// actually does the grid drawing
+	INTERACTIVETOOLSFRAMEWORK_API virtual void InternalDrawGrid(const FVector& Center, const FVector& Normal, const FVector& Up, const FVector& Side, const FVector& Size, const FLinearColor& Color, float LineThickness, bool bDepthTested);
 	// actually does the X drawing
 	INTERACTIVETOOLSFRAMEWORK_API virtual void InternalDrawViewFacingX(const FVector& Position, float Width, const FLinearColor& Color, float LineThickness, bool bDepthTested);
 	// actually does the disc drawing
