[ITF] Added value type aliases to GizmoPerStateValue structs

diff --git a/Engine/Source/Runtime/InteractiveToolsFramework/Public/BaseGizmos/GizmoElementShared.h b/Engine/Source/Runtime/InteractiveToolsFramework/Public/BaseGizmos/GizmoElementShared.h
--- a/Engine/Source/Runtime/InteractiveToolsFramework/Public/BaseGizmos/GizmoElementShared.h
+++ b/Engine/Source/Runtime/InteractiveToolsFramework/Public/BaseGizmos/GizmoElementShared.h
@@ -104,6 +104,8 @@ struct FGizmoPerStateValueDouble
 {
 	GENERATED_BODY()
 
+	using FValueType = double;
+
 	/**
 	 * Default value, used when the Interaction State is "None".
 	 * Optional to allow explicit un-setting, implying inheritance or some other value source.
@@ -172,6 +174,8 @@ struct FGizmoPerStateValueLinearColor
 {
 	GENERATED_BODY()
 
+	using FValueType = FLinearColor;
+
 	/**
 	 * Default value, used when the Interaction State is "None".
 	 * Optional to allow explicit un-setting, implying inheritance or some other value source.
