[ITF] Support custom cursor widgets, in addition to overriding the cursor enum

diff --git a/Engine/Source/Editor/UnrealEd/Private/EditorModeManager.cpp b/Engine/Source/Editor/UnrealEd/Private/EditorModeManager.cpp
--- a/Engine/Source/Editor/UnrealEd/Private/EditorModeManager.cpp
+++ b/Engine/Source/Editor/UnrealEd/Private/EditorModeManager.cpp
@@ -1801,6 +1801,16 @@ bool FEditorModeTools::GetCursor(EMouseCursor::Type& OutCursor) const
 	return bHandled;
 }
 
+TSharedPtr<SWidget> FEditorModeTools::MapCursor() const
+{
+	if (CVarEnableITFCursorOverrideSupport.GetValueOnGameThread() > 0)
+	{
+		return InteractiveToolsContext->GetCursorWidget();
+	}
+
+	return nullptr;
+}
+
 bool FEditorModeTools::GetOverrideCursorVisibility(bool& bWantsOverride, bool& bHardwareCursorVisible, bool bSoftwareCursorVisible) const
 {
 	bool bHandled = false;
diff --git a/Engine/Source/Editor/UnrealEd/Private/EditorViewportClient.cpp b/Engine/Source/Editor/UnrealEd/Private/EditorViewportClient.cpp
--- a/Engine/Source/Editor/UnrealEd/Private/EditorViewportClient.cpp
+++ b/Engine/Source/Editor/UnrealEd/Private/EditorViewportClient.cpp
@@ -1810,6 +1810,18 @@ EMouseCursor::Type FEditorViewportClient::GetCursor(FViewport* InViewport,int32
 	return MouseCursor;
 }
 
+TOptional<TSharedRef<SWidget>> FEditorViewportClient::MapCursor(FViewport* InViewport, const FCursorReply& CursorReply)
+{
+	// allow editor modes to override cursor
+	if (const TSharedPtr<SWidget> MappedCursor = ModeTools->MapCursor();
+		MappedCursor.IsValid())
+	{
+		return MappedCursor.ToSharedRef();
+	}
+	
+	return FCommonViewportClient::MapCursor(InViewport, CursorReply);
+}
+
 bool FEditorViewportClient::IsOrtho() const
 {
 	return !IsPerspective();
diff --git a/Engine/Source/Editor/UnrealEd/Private/Tools/EdModeInteractiveToolsContext.cpp b/Engine/Source/Editor/UnrealEd/Private/Tools/EdModeInteractiveToolsContext.cpp
--- a/Engine/Source/Editor/UnrealEd/Private/Tools/EdModeInteractiveToolsContext.cpp
+++ b/Engine/Source/Editor/UnrealEd/Private/Tools/EdModeInteractiveToolsContext.cpp
@@ -1210,6 +1210,19 @@ bool UModeManagerInteractiveToolsContext::GetCursor(EMouseCursor::Type& OutCurso
 	return false;
 }
 
+TSharedPtr<SWidget> UModeManagerInteractiveToolsContext::GetCursorWidget() const
+{
+	const UToolsContextCursorAPI* ToolsContextCursorAPI = ToolManager->GetContextObjectStore()->FindContext<UToolsContextCursorAPI>();
+	if (ToolsContextCursorAPI
+		&& ToolsContextCursorAPI->IsCursorOverridden()
+		&& ToolsContextCursorAPI->GetCurrentCursorOverride() == EMouseCursor::Custom)
+	{
+		return ToolsContextCursorAPI->GetCurrentCursorOverrideWidget();
+	}
+
+	return nullptr;
+}
+
 FRay UModeManagerInteractiveToolsContext::GetLastWorldRay() const
 {
 	return CurrentMouseState.Mouse.WorldRay;
diff --git a/Engine/Source/Editor/UnrealEd/Public/EditorModeManager.h b/Engine/Source/Editor/UnrealEd/Public/EditorModeManager.h
--- a/Engine/Source/Editor/UnrealEd/Public/EditorModeManager.h
+++ b/Engine/Source/Editor/UnrealEd/Public/EditorModeManager.h
@@ -305,6 +305,13 @@ public:
 	/** Get a cursor to override the default with, if any */
 	UNREALED_API bool GetCursor(EMouseCursor::Type& OutCursor) const;
 
+	/**
+	 * Called to map a cursor reply to an actual widget to render.
+	 *
+	 * @return	the widget that should be rendered for this cursor, if any.
+	 */
+	UNREALED_API TSharedPtr<SWidget> MapCursor() const;
+
 	/** Get override cursor visibility settings */
 	UNREALED_API bool GetOverrideCursorVisibility(bool& bWantsOverride, bool& bHardwareCursorVisible, bool bSoftwareCursorVisible) const;
 
diff --git a/Engine/Source/Editor/UnrealEd/Public/EditorViewportClient.h b/Engine/Source/Editor/UnrealEd/Public/EditorViewportClient.h
--- a/Engine/Source/Editor/UnrealEd/Public/EditorViewportClient.h
+++ b/Engine/Source/Editor/UnrealEd/Public/EditorViewportClient.h
@@ -608,6 +608,7 @@ public:
 	UNREALED_API virtual void MouseMove(FViewport* Viewport,int32 x, int32 y) override;
 	UNREALED_API virtual void MouseLeave( FViewport* Viewport ) override;
 	UNREALED_API virtual EMouseCursor::Type GetCursor(FViewport* Viewport,int32 X,int32 Y) override;
+	UNREALED_API virtual TOptional<TSharedRef<SWidget>> MapCursor(FViewport* Viewport, const FCursorReply& CursorReply) override;
 	UNREALED_API virtual void CapturedMouseMove( FViewport* InViewport, int32 InMouseX, int32 InMouseY ) override;
 	UNREALED_API virtual void ProcessAccumulatedPointerInput(FViewport* InViewport) override;
 	UNREALED_API virtual bool IsOrtho() const override;
diff --git a/Engine/Source/Editor/UnrealEd/Public/Tools/EdModeInteractiveToolsContext.h b/Engine/Source/Editor/UnrealEd/Public/Tools/EdModeInteractiveToolsContext.h
--- a/Engine/Source/Editor/UnrealEd/Public/Tools/EdModeInteractiveToolsContext.h
+++ b/Engine/Source/Editor/UnrealEd/Public/Tools/EdModeInteractiveToolsContext.h
@@ -304,6 +304,9 @@ public:
 	/** @return True if the context has overriden the cursor style, false if not. */
 	UE_API bool GetCursor(EMouseCursor::Type& OutCursor) const;
 
+	/** @return The Slate widget to use for the cursor if any. Only called if GetCursor returns Custom. */
+	UE_API TSharedPtr<SWidget> GetCursorWidget() const;
+
 	/** @return Ray into 3D scene at last mouse event */
 	UE_API virtual FRay GetLastWorldRay() const override;
 
diff --git a/Engine/Source/Runtime/InteractiveToolsFramework/Public/ToolContextInterfaces.h b/Engine/Source/Runtime/InteractiveToolsFramework/Public/ToolContextInterfaces.h
--- a/Engine/Source/Runtime/InteractiveToolsFramework/Public/ToolContextInterfaces.h
+++ b/Engine/Source/Runtime/InteractiveToolsFramework/Public/ToolContextInterfaces.h
@@ -28,6 +28,7 @@ class FSceneView;
 class FText;
 class FToolCommandChange;
 class FViewport;
+class SWidget;
 class UActorComponent;
 class UInteractiveGizmoManager;
 class UInteractiveToolManager;
@@ -192,6 +193,7 @@ class UToolsContextCursorAPI : public UObject
 
 private:
 	EMouseCursor::Type CursorOverrideShape;
+	TSharedPtr<SWidget> CursorOverrideWidget;
 	bool bIsCursorOverridden;
 
 public:
@@ -199,16 +201,22 @@ public:
 	UToolsContextCursorAPI()
 	{
 		CursorOverrideShape = EMouseCursor::Default;
+		CursorOverrideWidget = nullptr;
 		bIsCursorOverridden = false;
 	}
 
-	virtual ~UToolsContextCursorAPI() {}
+	virtual ~UToolsContextCursorAPI() override = default;
 
 	virtual EMouseCursor::Type GetCurrentCursorOverride() const 
 	{
 		return CursorOverrideShape;
 	}
 
+	virtual TSharedPtr<SWidget> GetCurrentCursorOverrideWidget() const 
+	{
+		return CursorOverrideWidget;
+	}
+
 	virtual bool IsCursorOverridden() const
 	{
 		return bIsCursorOverridden;
@@ -220,10 +228,19 @@ public:
 		CursorOverrideShape = CursorOverride;
 	}
 
+	virtual void SetCursorOverrideWidget(const TSharedRef<SWidget>& InCursorOverrideWidget)
+	{
+		bIsCursorOverridden = true;
+		CursorOverrideShape = EMouseCursor::Custom;
+		CursorOverrideWidget = InCursorOverrideWidget;
+	}
+
+	/** Clears the custom widget, if any, resets the shape to Default and resets the overridden flag. */
 	virtual void ClearCursorOverride() 
 	{
 		bIsCursorOverridden = false;
 		CursorOverrideShape = EMouseCursor::Default;
+		CursorOverrideWidget = nullptr;
 	}
 };
 
