[ITF] Account for (OS) DPI Scale in World->Pixel transforms

diff --git a/Engine/Source/Editor/UnrealEd/Private/Tools/EdModeInteractiveToolsContext.cpp b/Engine/Source/Editor/UnrealEd/Private/Tools/EdModeInteractiveToolsContext.cpp
--- a/Engine/Source/Editor/UnrealEd/Private/Tools/EdModeInteractiveToolsContext.cpp
+++ b/Engine/Source/Editor/UnrealEd/Private/Tools/EdModeInteractiveToolsContext.cpp
@@ -542,6 +542,7 @@ public:
 		ViewCameraState.Position = ViewTransform.GetLocation();
 		ViewCameraState.HorizontalFOVDegrees = ViewportClient->ViewFOV;
 		ViewCameraState.AspectRatio = ViewportClient->AspectRatio;
+		ViewCameraState.DPIScale = ViewportClient->GetDPIScale();
 
 		// ViewTransform rotation is only initialized for perspective!
 		if (ViewCameraState.bIsOrthographic == false)
@@ -608,6 +609,7 @@ void UEditorInteractiveToolsContext::Render(const FSceneView* View, FViewport* V
 	if (GizmoViewContext && ViewportClient == EditorModeManager->GetHoveredViewportClient())
 	{
 		GizmoViewContext->ResetFromSceneView(*View);
+		GizmoViewContext->SetDPIScale(ViewportClient->GetDPIScale());
 	}
 
 	// Render Tool and Gizmos
diff --git a/Engine/Source/Runtime/InteractiveToolsFramework/Private/BaseGizmos/GizmoRenderingUtil.cpp b/Engine/Source/Runtime/InteractiveToolsFramework/Private/BaseGizmos/GizmoRenderingUtil.cpp
--- a/Engine/Source/Runtime/InteractiveToolsFramework/Private/BaseGizmos/GizmoRenderingUtil.cpp
+++ b/Engine/Source/Runtime/InteractiveToolsFramework/Private/BaseGizmos/GizmoRenderingUtil.cpp
@@ -112,7 +112,8 @@ UViewAdjustedStaticMeshGizmoComponent* UE::GizmoRenderingUtil::CreateDefaultMate
 
 float UE::GizmoRenderingUtil::CalculateLocalPixelToWorldScale(
 	const FSceneView* View,
-	const FVector& Location)
+	const FVector& Location,
+	const double InDPIScale)
 {
 	using namespace GizmoRenderingUtilLocals;
 
@@ -122,12 +123,14 @@ float UE::GizmoRenderingUtil::CalculateLocalPixelToWorldScale(
 	}
 
 	UE::GizmoRenderingUtil::FSceneViewWrapper Wrapper(*View);
-	return CalculateLocalPixelToWorldScale(&Wrapper, Location);
+	return CalculateLocalPixelToWorldScale(&Wrapper, Location, true, InDPIScale);
 }
 
 float UE::GizmoRenderingUtil::CalculateLocalPixelToWorldScale(
 	const UE::GizmoRenderingUtil::ISceneViewInterface* View, 
-	const FVector& Location)
+	const FVector& Location,
+	const bool bInWithDPIScale,
+	const TOptional<double>& InDPIScale)
 {
 	using namespace GizmoRenderingUtilLocals;
 
@@ -144,9 +147,19 @@ float UE::GizmoRenderingUtil::CalculateLocalPixelToWorldScale(
 	FVector OffsetPointWorld = Location + OffsetDelta * View->GetViewRight() + OffsetDelta * View->GetViewUp();
 	FVector2D PixelB = WorldToPixelDouble(View, OffsetPointWorld);
 
+	constexpr double MinPixelToWorldScale = UE_SMALL_NUMBER;
+	constexpr double MaxPixelToWorldScale = 100.0; // Arbitrary large value to prevent division by zero or very small scale
+
 	double PixelDeltaSqr = VectorDifferenceSqr(PixelA, PixelB);
 	double WorldDeltaSqr = VectorDifferenceSqr(Location, OffsetPointWorld);
-	return (float)(sqrt(WorldDeltaSqr / PixelDeltaSqr));
+
+	double Result = FMath::Clamp(FMath::Sqrt(WorldDeltaSqr / PixelDeltaSqr), MinPixelToWorldScale, MaxPixelToWorldScale);
+	if (bInWithDPIScale)
+	{
+		Result *= InDPIScale.Get(View->GetDPIScale());
+	}
+
+	return static_cast<float>(Result);
 }
 
 float UE::GizmoRenderingUtil::CalculateViewDependentScaleAndFlatten(
diff --git a/Engine/Source/Runtime/InteractiveToolsFramework/Public/BaseGizmos/GizmoElementBase.h b/Engine/Source/Runtime/InteractiveToolsFramework/Public/BaseGizmos/GizmoElementBase.h
--- a/Engine/Source/Runtime/InteractiveToolsFramework/Public/BaseGizmos/GizmoElementBase.h
+++ b/Engine/Source/Runtime/InteractiveToolsFramework/Public/BaseGizmos/GizmoElementBase.h
@@ -51,7 +51,7 @@ public:
 		// Note: non-uniform scale is not supported and the X scale element will be used for uniform scaling.
 		FTransform LocalToWorldTransform;
 
-		// Pixel to world scale
+		// Pixel to world scale. To account for DPI scale, @see IToolsContextRenderAPI::ViewCameraState.
 		double PixelToWorldScale = 1.0;
 
 		// Interact state, if not equal to none, overrides the element's interact state 
@@ -64,10 +64,10 @@ public:
 		FGizmoElementLineRenderStateAttributes LineRenderState;
 
 		// Initialize state 
-		void Initialize(const FSceneView* InSceneView, FTransform InTransform)
+		void Initialize(const FSceneView* InSceneView, FTransform InTransform, const double InDPIScale = 1.0)
 		{
 			LocalToWorldTransform = InTransform;
-			PixelToWorldScale = UE::GizmoRenderingUtil::CalculateLocalPixelToWorldScale(InSceneView, InTransform.GetLocation());
+			PixelToWorldScale = UE::GizmoRenderingUtil::CalculateLocalPixelToWorldScale(InSceneView, InTransform.GetLocation(), InDPIScale);
 		}
 
 		// Returns the mesh material based on the current interaction state.
diff --git a/Engine/Source/Runtime/InteractiveToolsFramework/Public/BaseGizmos/GizmoRenderingUtil.h b/Engine/Source/Runtime/InteractiveToolsFramework/Public/BaseGizmos/GizmoRenderingUtil.h
--- a/Engine/Source/Runtime/InteractiveToolsFramework/Public/BaseGizmos/GizmoRenderingUtil.h
+++ b/Engine/Source/Runtime/InteractiveToolsFramework/Public/BaseGizmos/GizmoRenderingUtil.h
@@ -33,6 +33,7 @@ namespace UE::GizmoRenderingUtil
 		virtual const FMatrix& GetProjectionMatrix() const = 0;
 		virtual const FMatrix& GetViewMatrix() const = 0;
 		virtual bool IsPerspectiveProjection() const = 0;
+		virtual double GetDPIScale() const = 0;
 	};
 
 	// Wrapper around FSceneView to access it through ISceneViewInterface
@@ -53,6 +54,7 @@ namespace UE::GizmoRenderingUtil
 		virtual const FMatrix& GetProjectionMatrix() const { return View->ViewMatrices.GetProjectionMatrix(); }
 		virtual const FMatrix& GetViewMatrix() const { return View->ViewMatrices.GetViewMatrix(); }
 		virtual bool IsPerspectiveProjection() const { return View->IsPerspectiveProjection(); }
+		virtual double GetDPIScale() const { return 1.0; }
 
 	private:
 		const FSceneView* View;
@@ -109,10 +111,15 @@ namespace UE::GizmoRenderingUtil
 	 */
 	INTERACTIVETOOLSFRAMEWORK_API float CalculateLocalPixelToWorldScale(
 		const FSceneView* View,
-		const FVector& Location);
+		const FVector& Location,
+		const double InDPIScale = 1.0);
+
+	/** @note: if bInWithDPIScale is true, we use ISceneViewInterface::DPIScale or the provided InDPIScale when set.  */
 	INTERACTIVETOOLSFRAMEWORK_API float CalculateLocalPixelToWorldScale(
 		const UE::GizmoRenderingUtil::ISceneViewInterface* View,
-		const FVector& Location);
+		const FVector& Location,
+		const bool bInWithDPIScale = false,
+		const TOptional<double>& InDPIScale = TOptional<double>());
 
 	/**
 	 * @return Legacy view dependent conversion factor.
diff --git a/Engine/Source/Runtime/InteractiveToolsFramework/Public/BaseGizmos/GizmoViewContext.h b/Engine/Source/Runtime/InteractiveToolsFramework/Public/BaseGizmos/GizmoViewContext.h
--- a/Engine/Source/Runtime/InteractiveToolsFramework/Public/BaseGizmos/GizmoViewContext.h
+++ b/Engine/Source/Runtime/InteractiveToolsFramework/Public/BaseGizmos/GizmoViewContext.h
@@ -64,6 +64,8 @@ public:
 	const FMatrix& GetProjectionMatrix() const override { return ViewMatrices.GetProjectionMatrix(); }
 	const FMatrix& GetViewMatrix() const override { return ViewMatrices.GetViewMatrix(); }
 	bool IsPerspectiveProjection() const override { return bIsPerspectiveProjection; }
+	virtual double GetDPIScale() const override { return DPIScale; }
+	void SetDPIScale(const double InDPIScale) { DPIScale = InDPIScale; }
 
 	FVector4 WorldToScreen(const FVector& WorldPoint) const override
 	{
@@ -109,4 +111,5 @@ public:
 
 protected:
 	bool bIsPerspectiveProjection;
+	double DPIScale = 1.0; // Usually set from ViewportClient
 };
diff --git a/Engine/Source/Runtime/InteractiveToolsFramework/Public/ToolContextInterfaces.h b/Engine/Source/Runtime/InteractiveToolsFramework/Public/ToolContextInterfaces.h
--- a/Engine/Source/Runtime/InteractiveToolsFramework/Public/ToolContextInterfaces.h
+++ b/Engine/Source/Runtime/InteractiveToolsFramework/Public/ToolContextInterfaces.h
@@ -95,6 +95,8 @@ struct FViewCameraState
 	bool bIsOrthographic = false;
 	/** Is current view a VR view */
 	bool bIsVR = false;
+	/** DPI scale */
+	double DPIScale = 1.0;
 
 	/** @return "right"/horizontal direction in camera plane */
 	FVector Right() const { return Orientation.GetAxisY(); }
