[ITF] Adding missing support for LineTrace ViewDepthOffset

diff --git a/Engine/Source/Runtime/InteractiveToolsFramework/Private/BaseGizmos/GizmoElementBase.cpp b/Engine/Source/Runtime/InteractiveToolsFramework/Private/BaseGizmos/GizmoElementBase.cpp
--- a/Engine/Source/Runtime/InteractiveToolsFramework/Private/BaseGizmos/GizmoElementBase.cpp
+++ b/Engine/Source/Runtime/InteractiveToolsFramework/Private/BaseGizmos/GizmoElementBase.cpp
@@ -154,22 +154,47 @@ bool UGizmoElementBase::GetViewAlignRot(const FVector& InViewLocation, const FVe
 	return true;
 }
 
-void UGizmoElementBase::ApplyViewDepthOffset(const FVector& InViewLocation, const FVector& InViewDirection, FTransform& InOutTransform) const
+void UGizmoElementBase::ApplyViewDepthOffset(const FSceneView* InView, const double& InViewDepthOffset, FTransform& InOutTransform) const
 {
-	ApplyViewDepthOffset(InViewLocation, InViewDirection, ViewDepthOffset, InOutTransform);
+	bool bIsPerspectiveView;
+	FVector ViewLocation, ViewDirection, ViewUp;
+	GizmoElementBaseLocals::GetViewInfo<FSceneView>(InView, ViewLocation, ViewDirection, ViewUp, bIsPerspectiveView);
+	ApplyViewDepthOffset(ViewLocation, ViewDirection, bIsPerspectiveView, InViewDepthOffset, InOutTransform);
+}
+
+void UGizmoElementBase::ApplyViewDepthOffset(const FVector& InViewLocation, const FVector& InViewDirection, const bool bInIsPerspectiveView, FTransform& InOutTransform) const
+{
+	ApplyViewDepthOffset(InViewLocation, InViewDirection, bInIsPerspectiveView, ViewDepthOffset, InOutTransform);
 }
 
-void UGizmoElementBase::ApplyViewDepthOffset(const FVector& InViewLocation, const FVector& InViewDirection, const double& InViewDepthOffset, FTransform& InOutTransform) const
+void UGizmoElementBase::ApplyViewDepthOffset(const FVector& InViewLocation, const FVector& InViewDirection, const bool bInIsPerspectiveView, const double& InViewDepthOffset, FTransform& InOutTransform) const
 {
+	FVector ViewLocation = InViewLocation;
 	const FVector ObjectLocation = InOutTransform.GetLocation();
 
-	FVector CameraToObject = ObjectLocation - InViewLocation;
-	const double DistanceToObject = CameraToObject.Length();
+	FVector CameraToObject = ObjectLocation - ViewLocation;
+	double DistanceToObject = FVector::DotProduct(CameraToObject, InViewDirection);
+
+	const double ScaledViewDepthOffset = InViewDepthOffset * InOutTransform.GetScale3D().X;
+
+	// Orthographic view
+	if (!bInIsPerspectiveView)
+	{
+		// Back-project the effective view location based solely on the distance to the object along the view direction
+		// So that in screen space they'd overlap exactly
+		ViewLocation = ObjectLocation - InViewDirection * DistanceToObject;
+	}
+	else
+	{
+		DistanceToObject = CameraToObject.Length();
+	}
 
-	if (DistanceToObject < UE_KINDA_SMALL_NUMBER)
+	// Clamp to camera near plane
+	constexpr double MinDistance = UE_DOUBLE_KINDA_SMALL_NUMBER;
+	if (DistanceToObject < MinDistance)
 	{
 		// Object is at or behind the camera - move along the forward vector without accounting for scale
-		InOutTransform.SetLocation(InViewLocation + InViewDirection * InViewDepthOffset);
+		InOutTransform.SetLocation(ViewLocation + InViewDirection * ScaledViewDepthOffset);
 		return;
 	}
 
@@ -180,6 +205,12 @@ void UGizmoElementBase::ApplyViewDepthOffset(const FVector& InViewLocation, cons
 		TargetDistance = UE_KINDA_SMALL_NUMBER;
 	}
 
+	if (TargetDistance <= MinDistance)
+	{
+		// Prevent negative/zero
+		TargetDistance = 1;
+	}
+
 	// Calculate scale factor to maintain constant visual size
 	const double ScaleFactor = TargetDistance / DistanceToObject;
 
@@ -187,7 +218,7 @@ void UGizmoElementBase::ApplyViewDepthOffset(const FVector& InViewLocation, cons
 	CameraToObject.Normalize(); // Normalize the vector to get direction
 
 	// Update the transform with new location and scale
-	InOutTransform.SetLocation(InViewLocation + CameraToObject * TargetDistance);
+	InOutTransform.SetLocation(ViewLocation + CameraToObject * TargetDistance);
 	InOutTransform.SetScale3D(InOutTransform.GetScale3D() * ScaleFactor);
 }
 
@@ -313,12 +344,16 @@ bool UGizmoElementBase::UpdateRenderState(IToolsContextRenderAPI* RenderAPI, con
 
 	if (IsVisible(View, InOutRenderState.InteractionState, InOutRenderState.LocalToWorldTransform, InLocalCenter))
 	{
-		bOutHasAlignRot = GetViewAlignRot(View, InOutRenderState.LocalToWorldTransform, InLocalCenter, OutAlignRot);
+		bool bIsPerspectiveView;
+		FVector ViewLocation, ViewDirection, ViewUp;
+		GizmoElementBaseLocals::GetViewInfo<FSceneView>(View, ViewLocation, ViewDirection, ViewUp, bIsPerspectiveView);
+		bOutHasAlignRot = GetViewAlignRot(ViewLocation, ViewDirection, ViewUp, bIsPerspectiveView, InOutRenderState.LocalToWorldTransform, InLocalCenter, OutAlignRot);
+
 		InOutRenderState.LocalToWorldTransform = FTransform(OutAlignRot, InLocalCenter) * InOutRenderState.LocalToWorldTransform;
 
 		if (!FMath::IsNearlyZero(ViewDepthOffset))
 		{
-			ApplyViewDepthOffset(View->ViewLocation, View->GetViewDirection(), InOutRenderState.LocalToWorldTransform);
+			ApplyViewDepthOffset(ViewLocation, ViewDirection, bIsPerspectiveView, InOutRenderState.LocalToWorldTransform);
 		}
 	
 		return true;
@@ -343,8 +378,18 @@ bool UGizmoElementBase::UpdateLineTraceState(const UGizmoViewContext* ViewContex
 
 	if (IsHittable(ViewContext, InOutLineTraceState.LocalToWorldTransform, InLocalCenter))
 	{
-		bOutHasAlignRot = GetViewAlignRot(ViewContext, InOutLineTraceState.LocalToWorldTransform, InLocalCenter, OutAlignRot);
+		bool bIsPerspectiveView;
+		FVector ViewLocation, ViewDirection, ViewUp;
+		GizmoElementBaseLocals::GetViewInfo<UGizmoViewContext>(ViewContext, ViewLocation, ViewDirection, ViewUp, bIsPerspectiveView);
+		bOutHasAlignRot = GetViewAlignRot(ViewLocation, ViewDirection, ViewUp, bIsPerspectiveView, InOutLineTraceState.LocalToWorldTransform, InLocalCenter, OutAlignRot);
+
 		InOutLineTraceState.LocalToWorldTransform = FTransform(OutAlignRot, InLocalCenter) * InOutLineTraceState.LocalToWorldTransform;
+
+		if (!FMath::IsNearlyZero(ViewDepthOffset))
+		{
+			ApplyViewDepthOffset(ViewLocation, ViewDirection, bIsPerspectiveView, InOutLineTraceState.LocalToWorldTransform);
+		}
+
 		return true;
 	}
 
diff --git a/Engine/Source/Runtime/InteractiveToolsFramework/Private/BaseGizmos/GizmoElementCircle.cpp b/Engine/Source/Runtime/InteractiveToolsFramework/Private/BaseGizmos/GizmoElementCircle.cpp
--- a/Engine/Source/Runtime/InteractiveToolsFramework/Private/BaseGizmos/GizmoElementCircle.cpp
+++ b/Engine/Source/Runtime/InteractiveToolsFramework/Private/BaseGizmos/GizmoElementCircle.cpp
@@ -105,7 +105,7 @@ void UGizmoElementCircle::Render(IToolsContextRenderAPI* RenderAPI, const FRende
 				if (bDrawLine)
 				{
 					const double MeshViewDepthOffset = GetViewDepthOffset() - CurrentLineThickness;
-					ApplyViewDepthOffset(View->ViewLocation, View->GetViewDirection(), MeshViewDepthOffset, CurrentRenderState.LocalToWorldTransform);
+					ApplyViewDepthOffset(View, MeshViewDepthOffset, CurrentRenderState.LocalToWorldTransform);
 				}
 
 				const FColor VertexColor = CurrentRenderState.GetCurrentVertexColor().ToFColor(true);
diff --git a/Engine/Source/Runtime/InteractiveToolsFramework/Public/BaseGizmos/GizmoElementBase.h b/Engine/Source/Runtime/InteractiveToolsFramework/Public/BaseGizmos/GizmoElementBase.h
--- a/Engine/Source/Runtime/InteractiveToolsFramework/Public/BaseGizmos/GizmoElementBase.h
+++ b/Engine/Source/Runtime/InteractiveToolsFramework/Public/BaseGizmos/GizmoElementBase.h
@@ -576,9 +576,11 @@ protected:
 	// @param OutAlignRot the rotation to align this element in local space, should be prepended to the local-to-world transform.
 	INTERACTIVETOOLSFRAMEWORK_API virtual bool GetViewAlignRot(const FVector& InViewLocation, const FVector& InViewDirection, const FVector& InViewUp, bool bInPerspectiveView, const FTransform& InLocalToWorldTransform, const FVector& InLocalCenter, FQuat& OutAlignRot) const;
 
-	INTERACTIVETOOLSFRAMEWORK_API virtual void ApplyViewDepthOffset(const FVector& InViewLocation, const FVector& InViewDirection, FTransform& InOutTransform) const;
+	INTERACTIVETOOLSFRAMEWORK_API virtual void ApplyViewDepthOffset(const FSceneView* InView, const double& InViewDepthOffset, FTransform& InOutTransform) const;
 	
-	INTERACTIVETOOLSFRAMEWORK_API virtual void ApplyViewDepthOffset(const FVector& InViewLocation, const FVector& InViewDirection, const double& InViewDepthOffset, FTransform& InOutTransform) const;
+	INTERACTIVETOOLSFRAMEWORK_API virtual void ApplyViewDepthOffset(const FVector& InViewLocation, const FVector& InViewDirection, const bool bInIsPerspectiveView, FTransform& InOutTransform) const;
+
+	INTERACTIVETOOLSFRAMEWORK_API virtual void ApplyViewDepthOffset(const FVector& InViewLocation, const FVector& InViewDirection, const bool bInIsPerspectiveView, const double& InViewDepthOffset, FTransform& InOutTransform) const;
 
 	// Get mesh render state material attribute's value. 
 	INTERACTIVETOOLSFRAMEWORK_API UMaterialInterface* GetMaterial();
