[ITF] Added InvViewProjectionMatrix to GizmoViewContext

diff --git a/Engine/Source/Runtime/InteractiveToolsFramework/Public/BaseGizmos/GizmoViewContext.h b/Engine/Source/Runtime/InteractiveToolsFramework/Public/BaseGizmos/GizmoViewContext.h
--- a/Engine/Source/Runtime/InteractiveToolsFramework/Public/BaseGizmos/GizmoViewContext.h
+++ b/Engine/Source/Runtime/InteractiveToolsFramework/Public/BaseGizmos/GizmoViewContext.h
@@ -32,18 +32,21 @@ public:
 			ViewMatrix = SceneView.ViewMatrices.GetViewMatrix();
 			InvViewMatrix = SceneView.ViewMatrices.GetInvViewMatrix();
 			ViewProjectionMatrix = SceneView.ViewMatrices.GetViewProjectionMatrix();
+			InvViewProjectionMatrix = SceneView.ViewMatrices.GetInvViewProjectionMatrix();
 		}
 
 		const FMatrix& GetProjectionMatrix() const { return ProjectionMatrix; }
 		const FMatrix& GetViewMatrix() const { return ViewMatrix; }
 		const FMatrix& GetInvViewMatrix() const { return InvViewMatrix; }
 		const FMatrix& GetViewProjectionMatrix() const { return ViewProjectionMatrix; }
+		const FMatrix& GetInvViewProjectionMatrix() const { return InvViewProjectionMatrix; }
 
 	protected:
-		FMatrix ProjectionMatrix;
-		FMatrix ViewMatrix;
-		FMatrix InvViewMatrix;
-		FMatrix ViewProjectionMatrix;
+		FMatrix ProjectionMatrix = FMatrix::Identity;
+		FMatrix ViewMatrix = FMatrix::Identity;
+		FMatrix InvViewMatrix = FMatrix::Identity;
+		FMatrix ViewProjectionMatrix = FMatrix::Identity;
+		FMatrix InvViewProjectionMatrix = FMatrix::Identity;
 	};
 
 	// Use this to reinitialize the object each frame for the hovered viewport.
