[ITF] Added FGizmoElementTraverser util, for friended access

diff --git a/Engine/Source/Runtime/InteractiveToolsFramework/Private/BaseGizmos/GizmoElementBase.cpp b/Engine/Source/Runtime/InteractiveToolsFramework/Private/BaseGizmos/GizmoElementBase.cpp
--- a/Engine/Source/Runtime/InteractiveToolsFramework/Private/BaseGizmos/GizmoElementBase.cpp
+++ b/Engine/Source/Runtime/InteractiveToolsFramework/Private/BaseGizmos/GizmoElementBase.cpp
@@ -302,6 +302,86 @@ bool UGizmoElementBase::GetEnabledForViewProjection(bool bIsPerspectiveProjectio
 	return false;
 }
 
+void UGizmoElementBase::ForEachSubElement(const TFunctionRef<void(UGizmoElementBase*)>& InFunc)
+{
+	for (UGizmoElementBase* Element : GetSubElements())
+	{
+		if (Element)
+		{
+			InFunc(Element);
+		}
+	}
+}
+
+void UGizmoElementBase::ForEachSubElementRecursive(const TFunctionRef<void(UGizmoElementBase*)>& InFunc) const
+{
+	TArray<TObjectPtr<UGizmoElementBase>> SubElements;
+	if (GetSubElementsRecursive(SubElements))
+	{
+		for (UGizmoElementBase* Element : SubElements)
+		{
+			if (Element)
+			{
+				InFunc(Element);
+			}
+		}
+	}
+}
+
+void UGizmoElementBase::ForEachSubElementRecursive(const TFunctionRef<void(UGizmoElementBase*)>& InFunc, const uint32 InPartId) const
+{
+	TArray<TObjectPtr<UGizmoElementBase>> SubElements;
+	if (GetSubElementsRecursive(SubElements, InPartId))
+	{
+		for (UGizmoElementBase* Element : SubElements)
+		{
+			if (Element)
+			{
+				InFunc(Element);
+			}
+		}
+	}
+}
+
+TConstArrayView<TObjectPtr<UGizmoElementBase>> UGizmoElementBase::GetSubElements() const
+{
+	// The default element has no sub-elements, so we do nothing and return and empty result
+	return { };
+}
+
+bool UGizmoElementBase::GetSubElementsRecursive(TArray<TObjectPtr<UGizmoElementBase>>& OutElements) const
+{
+	OutElements.Append(GetSubElements());
+
+	for (const TObjectPtr<UGizmoElementBase>& Element : GetSubElements())
+	{
+		if (Element)
+		{
+			Element->GetSubElementsRecursive(OutElements);
+		}
+	}
+
+	return !OutElements.IsEmpty();
+}
+
+bool UGizmoElementBase::GetSubElementsRecursive(TArray<TObjectPtr<UGizmoElementBase>>& OutElements, const uint32 InPartId) const
+{
+	for (const TObjectPtr<UGizmoElementBase>& Element : GetSubElements())
+	{
+		if (Element)
+		{
+			if (Element->GetPartIdentifier() == InPartId)
+			{
+				OutElements.Add(Element);	
+			}
+
+			Element->GetSubElementsRecursive(OutElements, InPartId);
+		}
+	}
+
+	return !OutElements.IsEmpty();
+}
+
 bool UGizmoElementBase::IsVisible(const FSceneView* View, EGizmoElementInteractionState InCurrentInteractionState,
 	const FTransform& InLocalToWorldTransform, const FVector& InLocalCenter) const
 {
@@ -525,6 +605,26 @@ void UGizmoElementBase::SetPartIdentifier(uint32 InPartId, const bool bInOverrid
 	}
 }
 
+UGizmoElementBase* UGizmoElementBase::FindPartElement(const uint32 InPartId)
+{
+	if (GetPartIdentifier() == InPartId)
+	{
+		return this; // Return this element if the identifier matches
+	}
+
+	return nullptr;
+}
+
+const UGizmoElementBase* UGizmoElementBase::FindPartElement(const uint32 InPartId) const
+{
+	if (GetPartIdentifier() == InPartId)
+	{
+		return this; // Return this element if the identifier matches
+	}
+
+	return nullptr;
+}
+
 void UGizmoElementBase::SetElementState(EGizmoElementState InElementState)
 {
 	ElementState = InElementState;
diff --git a/Engine/Source/Runtime/InteractiveToolsFramework/Private/BaseGizmos/GizmoElementGroup.cpp b/Engine/Source/Runtime/InteractiveToolsFramework/Private/BaseGizmos/GizmoElementGroup.cpp
--- a/Engine/Source/Runtime/InteractiveToolsFramework/Private/BaseGizmos/GizmoElementGroup.cpp
+++ b/Engine/Source/Runtime/InteractiveToolsFramework/Private/BaseGizmos/GizmoElementGroup.cpp
@@ -6,9 +6,14 @@
 #include UE_INLINE_GENERATED_CPP_BY_NAME(GizmoElementGroup)
 
 void UGizmoElementGroupBase::ApplyUniformConstantScaleToTransform(double PixelToWorldScale, FTransform& InOutLocalToWorldTransform) const
+{
+	return ApplyUniformConstantScaleToTransform(PixelToWorldScale, InOutLocalToWorldTransform, false);
+}
+
+void UGizmoElementGroupBase::ApplyUniformConstantScaleToTransform(double PixelToWorldScale, FTransform& InOutLocalToWorldTransform, const bool bInForceApply) const
 {
 	double Scale = InOutLocalToWorldTransform.GetScale3D().X;
-	if (bConstantScale)
+	if (bInForceApply || bConstantScale)
 	{
 		Scale *= PixelToWorldScale;
 	}
@@ -26,7 +31,7 @@ void UGizmoElementGroupBase::Render(IToolsContextRenderAPI* RenderAPI, const FRe
 
 		// Continue render even if not visible so all transforms will be cached 
 		// for subsequent line tracing.
-		ForEachElement([RenderAPI, &CurrentRenderState](UGizmoElementBase* Element)
+		ForEachSubElement([RenderAPI, &CurrentRenderState](UGizmoElementBase* Element)
 		{
 			Element->Render(RenderAPI, CurrentRenderState);
 		});
@@ -44,7 +49,7 @@ void UGizmoElementGroupBase::DrawHUD(FCanvas* Canvas, IToolsContextRenderAPI* Re
 
 		// Continue render even if not visible so all transforms will be cached
 		// for subsequent line tracing.
-		ForEachElement([Canvas, RenderAPI, &CurrentRenderState](UGizmoElementBase* Element)
+		ForEachSubElement([Canvas, RenderAPI, &CurrentRenderState](UGizmoElementBase* Element)
 		{
 			Element->DrawHUD(Canvas, RenderAPI, CurrentRenderState);
 		});
@@ -62,7 +67,7 @@ FInputRayHit UGizmoElementGroupBase::LineTrace(const UGizmoViewContext* ViewCont
 	{
 		ApplyUniformConstantScaleToTransform(CurrentLineTraceState.PixelToWorldScale, CurrentLineTraceState.LocalToWorldTransform);
 
-		ForEachElement([this, ViewContext, &CurrentLineTraceState, &RayOrigin, &RayDirection, &Hit](UGizmoElementBase* Element)
+		ForEachSubElement([this, ViewContext, &CurrentLineTraceState, &RayOrigin, &RayDirection, &Hit](UGizmoElementBase* Element)
 		{
 			FInputRayHit NewHit = Element->LineTrace(ViewContext, CurrentLineTraceState, RayOrigin, RayDirection);
 			if (!Hit.bHit || NewHit.HitDepth < Hit.HitDepth)
@@ -89,10 +94,10 @@ void UGizmoElementGroupBase::SetPartIdentifier(uint32 InPartId, const bool bInOv
 	// Only set if the part identifier is current Default (0), and bInOverrideSet is true 
 	if (bInOverrideSet || GetPartIdentifier() == 0)
 	{
-		Super::SetPartIdentifier(InPartId);
+		UGizmoElementBase::SetPartIdentifier(InPartId);
 	}
 
-	ForEachElement([InPartId, bInOverrideUnsetChildren](UGizmoElementBase* Element)
+	ForEachSubElement([InPartId, bInOverrideUnsetChildren](UGizmoElementBase* Element)
 	{
 		// By default, bOverrideChildIfSet is true, but we only want that for the root element
 		constexpr bool bOverrideChildIfSet = false;
@@ -100,6 +105,32 @@ void UGizmoElementGroupBase::SetPartIdentifier(uint32 InPartId, const bool bInOv
 	});
 }
 
+UGizmoElementBase* UGizmoElementGroupBase::FindPartElement(const uint32 InPartId)
+{
+	for (UGizmoElementBase* Element : Elements)
+	{
+		if (Element && Element->GetPartIdentifier() == InPartId)
+		{
+			return Element;
+		}
+	}
+
+	return nullptr;
+}
+
+const UGizmoElementBase* UGizmoElementGroupBase::FindPartElement(const uint32 InPartId) const
+{
+	for (UGizmoElementBase* Element : Elements)
+	{
+		if (Element && Element->GetPartIdentifier() == InPartId)
+		{
+			return Element;
+		}
+	}
+
+	return nullptr;
+}
+
 void UGizmoElementGroupBase::Add(UGizmoElementBase* InElement)
 {
 	if (!Elements.Contains(InElement))
@@ -120,15 +151,9 @@ void UGizmoElementGroupBase::Remove(UGizmoElementBase* InElement)
 	}
 }
 
-void UGizmoElementGroupBase::ForEachElement(const TFunctionRef<void(UGizmoElementBase*)>& InFunc)
+TConstArrayView<TObjectPtr<UGizmoElementBase>> UGizmoElementGroupBase::GetSubElements() const
 {
-	for (UGizmoElementBase* Element : Elements)
-	{
-		if (Element)
-		{
-			InFunc(Element);
-		}
-	}
+	return Elements;
 }
 
 bool UGizmoElementGroupBase::UpdatePartVisibleState(bool bVisible, uint32 InPartIdentifier, bool bInAllowMultipleElements)
diff --git a/Engine/Source/Runtime/InteractiveToolsFramework/Private/BaseGizmos/GizmoUtil.cpp b/Engine/Source/Runtime/InteractiveToolsFramework/Private/BaseGizmos/GizmoUtil.cpp
--- a/Engine/Source/Runtime/InteractiveToolsFramework/Private/BaseGizmos/GizmoUtil.cpp
+++ b/Engine/Source/Runtime/InteractiveToolsFramework/Private/BaseGizmos/GizmoUtil.cpp
@@ -2,6 +2,8 @@
 
 #include "BaseGizmos/GizmoUtil.h"
 
+#include "BaseGizmos/GizmoElementCircleBase.h"
+#include "BaseGizmos/GizmoElementGroup.h"
 #include "InteractiveGizmoManager.h"
 #include "InteractiveGizmo.h"
 
@@ -26,3 +28,51 @@ UInteractiveGizmo* UE::GizmoUtil::CreateGizmoViaSimpleBuilder(UInteractiveGizmoM
 
 	return Gizmo;
 }
+
+TConstArrayView<TObjectPtr<UGizmoElementBase>> FGizmoElementAccessor::GetSubElements(const UGizmoElementBase& InGizmoElement)
+{
+	return InGizmoElement.GetSubElements();
+}
+
+void FGizmoElementAccessor::GetSubElementsRecursive(const UGizmoElementBase& InGizmoElement, TArray<TObjectPtr<UGizmoElementBase>>& OutElements)
+{
+	InGizmoElement.GetSubElementsRecursive(OutElements);
+}
+
+bool FGizmoElementAccessor::GetSubElementsRecursive(const UGizmoElementBase& InGizmoElement, TArray<TObjectPtr<UGizmoElementBase>>& OutElements, const uint32 InPartId)
+{
+	return InGizmoElement.GetSubElementsRecursive(OutElements, InPartId);
+}
+
+bool FGizmoElementAccessor::UpdateRenderState(UGizmoElementBase& InGizmoElement, IToolsContextRenderAPI* InRenderAPI, const FVector& InLocalCenter, UGizmoElementBase::FRenderTraversalState& InOutRenderState) const
+{
+	FQuat AlignRot;
+	bool bHasAlignRot;
+	return UpdateRenderState(InGizmoElement, InRenderAPI, InLocalCenter, InOutRenderState, bHasAlignRot, AlignRot);
+}
+
+bool FGizmoElementAccessor::UpdateRenderState(UGizmoElementBase& InGizmoElement, IToolsContextRenderAPI* InRenderAPI, const FVector& InLocalCenter, UGizmoElementBase::FRenderTraversalState& InOutRenderState, bool& bOutHasAlignRot, FQuat& OutAlignRot) const
+{
+	if (!ensure(InRenderAPI))
+	{
+		return false;
+	}
+
+	return InGizmoElement.UpdateRenderState(InRenderAPI, InLocalCenter, InOutRenderState, bOutHasAlignRot, OutAlignRot);
+}
+
+bool FGizmoElementAccessor::IsEnabledForInteractionState(const UGizmoElementBase& InGizmoElement, const TOptional<EGizmoElementInteractionState>& InInteractionState)
+{
+	return InGizmoElement.GetEnabledForInteractionState(InInteractionState.Get(InGizmoElement.GetElementInteractionState()));
+}
+
+bool FGizmoElementAccessor::IsPartial(
+	UGizmoElementCircleBase& InGizmoElement,
+	const FVector& InWorldCenter,
+	const FVector& InWorldNormal,
+	const FVector& InViewLocation,
+	const FVector& InViewDirection,
+	const bool bIsPerspectiveProjection)
+{
+	return InGizmoElement.IsPartial(InWorldCenter, InWorldNormal, InViewLocation, InViewDirection, bIsPerspectiveProjection);
+}
diff --git a/Engine/Source/Runtime/InteractiveToolsFramework/Public/BaseGizmos/GizmoElementBase.h b/Engine/Source/Runtime/InteractiveToolsFramework/Public/BaseGizmos/GizmoElementBase.h
--- a/Engine/Source/Runtime/InteractiveToolsFramework/Public/BaseGizmos/GizmoElementBase.h
+++ b/Engine/Source/Runtime/InteractiveToolsFramework/Public/BaseGizmos/GizmoElementBase.h
@@ -25,6 +25,9 @@ UCLASS(Transient, Abstract, MinimalAPI)
 class UGizmoElementBase : public UObject
 {
 	GENERATED_BODY()
+
+	friend class FGizmoElementAccessor;
+	
 public:
 
 	static constexpr float DefaultViewDependentAngleTol = 0.052f;			// ~3 degrees
@@ -176,6 +179,15 @@ public:
 	// If bInOverrideSet is true, the PartId will be set even if it's currently 0. 
 	INTERACTIVETOOLSFRAMEWORK_API virtual void SetPartIdentifier(uint32 InPartId, const bool bInOverrideUnsetChildren, const bool bInOverrideSet = true);
 
+	// Returns the element associated with the specified part id, or nullptr if not found.
+	INTERACTIVETOOLSFRAMEWORK_API virtual UGizmoElementBase* FindPartElement(const uint32 InPartId);
+
+	// Returns the element associated with the specified part id, or nullptr if not found.
+	INTERACTIVETOOLSFRAMEWORK_API virtual const UGizmoElementBase* FindPartElement(const uint32 InPartId) const;
+
+	// Get the (top-level) sub-elements, if any.
+	INTERACTIVETOOLSFRAMEWORK_API virtual TConstArrayView<TObjectPtr<UGizmoElementBase>> GetSubElements() const;
+
 	// Object type bitmask indicating whether this object is visible or hittable or both
 	INTERACTIVETOOLSFRAMEWORK_API virtual void SetElementState(EGizmoElementState InElementState);
 	INTERACTIVETOOLSFRAMEWORK_API virtual EGizmoElementState GetElementState() const;
@@ -542,6 +554,21 @@ private:
 
 protected:
 
+	// Iterate over each top-level element (sub-element), excluding the element itself.
+	INTERACTIVETOOLSFRAMEWORK_API void ForEachSubElement(const TFunctionRef<void(UGizmoElementBase*)>& InFunc);
+
+	// Iterate over each element (sub-element) in the element tree, excluding the element itself.
+	INTERACTIVETOOLSFRAMEWORK_API void ForEachSubElementRecursive(const TFunctionRef<void(UGizmoElementBase*)>& InFunc) const;
+
+	// Iterate over each sub-element with the given Part Id in the element tree, excluding the element itself.
+	INTERACTIVETOOLSFRAMEWORK_API void ForEachSubElementRecursive(const TFunctionRef<void(UGizmoElementBase*)>& InFunc, const uint32 InPartId) const;
+
+	// Get all sub-elements for the element tree, if any. Returns true if any sub-elements were found.
+	INTERACTIVETOOLSFRAMEWORK_API bool GetSubElementsRecursive(TArray<TObjectPtr<UGizmoElementBase>>& OutElements) const;
+
+	// Get all sub-elements with the given Part Id for the element tree, if any. Returns true if any sub-elements were found.
+	INTERACTIVETOOLSFRAMEWORK_API bool GetSubElementsRecursive(TArray<TObjectPtr<UGizmoElementBase>>& OutElements, const uint32 InPartId) const;
+
 	// Return whether element is currently visible.
 	INTERACTIVETOOLSFRAMEWORK_API virtual bool IsVisible(const FSceneView* View, EGizmoElementInteractionState InCurrentInteractionState, 
 		const FTransform& InLocalToWorldTransform, const FVector& InLocalCenter) const;
diff --git a/Engine/Source/Runtime/InteractiveToolsFramework/Public/BaseGizmos/GizmoElementGroup.h b/Engine/Source/Runtime/InteractiveToolsFramework/Public/BaseGizmos/GizmoElementGroup.h
--- a/Engine/Source/Runtime/InteractiveToolsFramework/Public/BaseGizmos/GizmoElementGroup.h
+++ b/Engine/Source/Runtime/InteractiveToolsFramework/Public/BaseGizmos/GizmoElementGroup.h
@@ -17,7 +17,7 @@ class UGizmoElementGroupBase : public UGizmoElementLineBase
 {
 	GENERATED_BODY()
 
-	friend class FGizmoElementTraverser;
+	friend class FGizmoElementAccessor;
 
 public:
 
@@ -27,11 +27,20 @@ public:
 	INTERACTIVETOOLSFRAMEWORK_API virtual FInputRayHit LineTrace(const UGizmoViewContext* ViewContext, const FLineTraceTraversalState& LineTraceState, const FVector& RayOrigin, const FVector& RayDirection) override;
 
 	// Sets the part identifier for this element and all of its children whose current ID is 0/Default.
-	using UGizmoElementLineBase::SetPartIdentifier;
+	using UGizmoElementBase::SetPartIdentifier;
 
 	// Sets the part identifier for this element and all of its children whose current ID is 0/Default.
 	INTERACTIVETOOLSFRAMEWORK_API virtual void SetPartIdentifier(uint32 InPartId, const bool bInOverrideUnsetChildren, const bool bInOverrideSet = true) override;
 
+	// Returns the element associated with the specified part id, or nullptr if not found.
+	INTERACTIVETOOLSFRAMEWORK_API virtual UGizmoElementBase* FindPartElement(const uint32 InPartId) override;
+
+	// Returns the element associated with the specified part id, or nullptr if not found.
+	INTERACTIVETOOLSFRAMEWORK_API virtual const UGizmoElementBase* FindPartElement(const uint32 InPartId) const override;
+
+	// Get the (top-level) sub-elements, if any.
+	INTERACTIVETOOLSFRAMEWORK_API virtual TConstArrayView<TObjectPtr<UGizmoElementBase>> GetSubElements() const override;
+
 	// Update group and contained elements' visibility state for elements in specified gizmo parts, return true if part was found.
 	// Setting bAllowMultipleElements to true will allow multiple elements with the same Part Id to be updated.
 	INTERACTIVETOOLSFRAMEWORK_API virtual bool UpdatePartVisibleState(bool bVisible, uint32 InPartIdentifier, bool bInAllowMultipleElements = false) override;
@@ -69,9 +78,6 @@ protected:
 	// Remove object from group, if it exists
 	INTERACTIVETOOLSFRAMEWORK_API virtual void Remove(UGizmoElementBase* InElement);
 
-	// Iterate over each element in the group
-	INTERACTIVETOOLSFRAMEWORK_API void ForEachElement(const TFunctionRef<void(UGizmoElementBase*)>& InFunc);
-
 protected:
 
 	// When true, maintains view-dependent constant scale for this gizmo object hierarchy
@@ -89,6 +95,9 @@ protected:
 
 	// Updates input transform's scale component to have uniform scale and applies constant scale if bConstantScale is true
 	INTERACTIVETOOLSFRAMEWORK_API virtual void ApplyUniformConstantScaleToTransform(double PixelToWorldScale, FTransform& InOutLocalToWorldTransform) const;
+
+	// Updates input transform's scale component to have uniform scale and applies constant scale if bConstantScale or bInForceApply is true
+	INTERACTIVETOOLSFRAMEWORK_API virtual void ApplyUniformConstantScaleToTransform(double PixelToWorldScale, FTransform& InOutLocalToWorldTransform, const bool bInForceApply) const;
 };
 
 /**
diff --git a/Engine/Source/Runtime/InteractiveToolsFramework/Public/BaseGizmos/GizmoUtil.h b/Engine/Source/Runtime/InteractiveToolsFramework/Public/BaseGizmos/GizmoUtil.h
--- a/Engine/Source/Runtime/InteractiveToolsFramework/Public/BaseGizmos/GizmoUtil.h
+++ b/Engine/Source/Runtime/InteractiveToolsFramework/Public/BaseGizmos/GizmoUtil.h
@@ -2,6 +2,7 @@
 
 #pragma once
 
+#include "GizmoElementBase.h"
 #include "InteractiveGizmoBuilder.h"
 #include "Templates/Function.h"
 #include "Templates/SubclassOf.h"
@@ -9,9 +10,36 @@
 #include "GizmoUtil.generated.h"
 
 class FString;
+class UGizmoElementCircleBase;
 class UInteractiveGizmo;
 class UInteractiveGizmoManager;
 
+/** A utility class to access GizmoElements that aren't otherwise publicly exposed. Use responsibly. */
+class FGizmoElementAccessor
+{
+public:
+	INTERACTIVETOOLSFRAMEWORK_API static TConstArrayView<TObjectPtr<UGizmoElementBase>> GetSubElements(const UGizmoElementBase& InGizmoElement);
+	INTERACTIVETOOLSFRAMEWORK_API static void GetSubElementsRecursive(const UGizmoElementBase& InGizmoElement, TArray<TObjectPtr<UGizmoElementBase>>& OutElements);
+	INTERACTIVETOOLSFRAMEWORK_API static bool GetSubElementsRecursive(const UGizmoElementBase& InGizmoElement, TArray<TObjectPtr<UGizmoElementBase>>& OutElements, const uint32 InPartId);
+
+	/** Updates the given render state from the given GizmoElement. */
+	INTERACTIVETOOLSFRAMEWORK_API bool UpdateRenderState(UGizmoElementBase& InGizmoElement, IToolsContextRenderAPI* InRenderAPI, const FVector& InLocalCenter, UGizmoElementBase::FRenderTraversalState& InOutRenderState) const;
+
+	/** Updates the given render state from the given GizmoElement. */
+	INTERACTIVETOOLSFRAMEWORK_API bool UpdateRenderState(UGizmoElementBase& InGizmoElement, IToolsContextRenderAPI* InRenderAPI, const FVector& InLocalCenter, UGizmoElementBase::FRenderTraversalState& InOutRenderState, bool& bOutHasAlignRot, FQuat& OutAlignRot) const;
+
+	/** Returns whether element is enabled for given interaction state. If the provided interaction state isn't set, it will use the currently applied state to the given element. */
+	INTERACTIVETOOLSFRAMEWORK_API static bool IsEnabledForInteractionState(const UGizmoElementBase& InGizmoElement, const TOptional<EGizmoElementInteractionState>& InInteractionState = {});
+
+	INTERACTIVETOOLSFRAMEWORK_API static bool IsPartial(
+		UGizmoElementCircleBase& InGizmoElement,
+		const FVector& InWorldCenter,
+		const FVector& InWorldNormal,
+		const FVector& InViewLocation,
+		const FVector& InViewDirection,
+		const bool bIsPerspectiveProjection);
+};
+
 /**
  * Gizmo builder that simply calls a particular lambda when building a gizmo. Makes it easy to
  *  register gizmo build behavior without writing a new class.
