InteractiveToolsFramework: Support creation of tool targets from combined inputs.

diff --git a/Engine/Source/Runtime/InteractiveToolsFramework/Private/ToolTargetManager.cpp b/Engine/Source/Runtime/InteractiveToolsFramework/Private/ToolTargetManager.cpp
--- a/Engine/Source/Runtime/InteractiveToolsFramework/Private/ToolTargetManager.cpp
+++ b/Engine/Source/Runtime/InteractiveToolsFramework/Private/ToolTargetManager.cpp
@@ -7,6 +7,31 @@
 
 #include UE_INLINE_GENERATED_CPP_BY_NAME(ToolTargetManager)
 
+namespace ToolTargetManagerLocals
+{
+	// Given an array and a 1:1 array of bools, remove all entries that have the corresponding bool set to true
+	void RemoveByFlags(TArray<UObject*>& ObjectsInOut, const TArray<bool>& RemoveFlags)
+	{
+		if (!ensure(RemoveFlags.Num() == ObjectsInOut.Num()))
+		{
+			return;
+		}
+
+		int32 OutIndex = 0;
+		int32 InIndex = 0;
+		while (InIndex < ObjectsInOut.Num())
+		{
+			if (!RemoveFlags[InIndex])
+			{
+				ObjectsInOut[OutIndex] = ObjectsInOut[InIndex];
+				++OutIndex;
+			}
+			++InIndex;
+		}
+		ObjectsInOut.SetNum(OutIndex);
+	}
+}
+
 void UToolTargetManager::Initialize()
 {
 	bIsActive = true;
@@ -78,10 +103,10 @@ UToolTarget* UToolTargetManager::BuildTarget(UObject* SourceObject, const FToolT
 
 int32 UToolTargetManager::CountSelectedAndTargetable(const FToolBuilderState& SceneState, const FToolTargetTypeRequirements& TargetType) const
 {
-	return ToolBuilderUtil::CountComponents(SceneState, [&](UActorComponent* Object)
-		{
-			return CanBuildTarget(Object, TargetType);
-		});
+	return CountSelectedAndTargetableWithPredicate(SceneState, TargetType, 
+		// We currently only ever get components out of our scene state, so it is ok to
+		//  use the predicate that takes an actor component
+		[](UActorComponent& Component) { return true; });
 }
 
 
@@ -89,49 +114,117 @@ void UToolTargetManager::EnumerateSelectedAndTargetableComponents(const FToolBui
 	const FToolTargetTypeRequirements& TargetRequirements,
 	TFunctionRef<void(UActorComponent*)> ComponentFunc) const
 {
-	ToolBuilderUtil::EnumerateComponents(SceneState, [&](UActorComponent* Component)
+	// Gather input objects
+	TArray<UObject*> InputObjects;
+	InputObjects.Append(ToolBuilderUtil::FindAllComponents(SceneState, [this](UActorComponent* Component)
+	{
+		return Component && (!InputFilterFunction || InputFilterFunction(Component));
+	}));
+
+	TArray<bool> Used;
+	for (TObjectPtr<UToolTargetFactory> Factory : Factories)
 	{
-		if (CanBuildTarget(Component, TargetRequirements))
+		int32 NumTargets = Factory->CanBuildTargets(InputObjects, TargetRequirements, Used);
+		if (NumTargets > 0)
 		{
-			ComponentFunc(Component);
+			for (int32 Index = 0; Index < Used.Num(); ++Index)
+			{
+				if (Used[Index] && ensure(InputObjects.IsValidIndex(Index)))
+				{
+					ComponentFunc(static_cast<UActorComponent*>(InputObjects[Index]));
+				}
+			}
+			ToolTargetManagerLocals::RemoveByFlags(InputObjects, Used);
+
+			if (InputObjects.IsEmpty())
+			{
+				break;
+			}
 		}
-	});
+	}
 }
 
-
 int32 UToolTargetManager::CountSelectedAndTargetableWithPredicate(const FToolBuilderState& SceneState,
 	const FToolTargetTypeRequirements& TargetRequirements,
 	TFunctionRef<bool(UActorComponent&)> ComponentPred) const
 {
-	return ToolBuilderUtil::CountComponents(SceneState, [&](UActorComponent* Object)
+	// Gather input objects
+	TArray<UObject*> InputObjects;
+	InputObjects.Append(ToolBuilderUtil::FindAllComponents(SceneState, [this, &ComponentPred](UActorComponent* Component)
 	{
-		return CanBuildTarget(Object, TargetRequirements) && ComponentPred(*Object);
-	});
+		return Component 
+			&& (!InputFilterFunction || InputFilterFunction(Component)) 
+			&& ComponentPred(*Component);
+	}));
+
+	int32 Count = 0;
+	TArray<bool> Used;
+	for (TObjectPtr<UToolTargetFactory> Factory : Factories)
+	{
+		int32 NumTargets = Factory->CanBuildTargets(InputObjects, TargetRequirements, Used);
+		if (NumTargets > 0)
+		{
+			Count += NumTargets;
+			ToolTargetManagerLocals::RemoveByFlags(InputObjects, Used);
+			if (InputObjects.IsEmpty())
+			{
+				break;
+			}
+		}
+	}
+
+	return Count;
 }
 
 UToolTarget* UToolTargetManager::BuildFirstSelectedTargetable(const FToolBuilderState& SceneState, const FToolTargetTypeRequirements& TargetType)
 {
-	return BuildTarget(
-		ToolBuilderUtil::FindFirstComponent(SceneState, [&](UActorComponent* Object)
+	// Gather input objects
+	TArray<UObject*> InputObjects;
+	InputObjects.Append(ToolBuilderUtil::FindAllComponents(SceneState, [this](UActorComponent* Component)
+	{
+		return Component && (!InputFilterFunction || InputFilterFunction(Component));
+	}));
+
+	TArray<bool> Used;
+	for (TObjectPtr<UToolTargetFactory> Factory : Factories)
+	{
+		int32 NumTargets = Factory->CanBuildTargets(InputObjects, TargetType, Used);
+		if (NumTargets > 0)
 		{
-			return CanBuildTarget(Object, TargetType);
-		}),
-		TargetType);
+			return Factory->BuildFirstTarget(InputObjects, TargetType, Used);
+		}
+	}
+	return nullptr;
 }
 
 TArray<TObjectPtr<UToolTarget>> UToolTargetManager::BuildAllSelectedTargetable(const FToolBuilderState& SceneState,
 	const FToolTargetTypeRequirements& TargetType)
 {
-	TArray<UActorComponent*> Components = ToolBuilderUtil::FindAllComponents(SceneState, [&](UActorComponent* Object)
+	TArray<TObjectPtr<UToolTarget>> TargetsOut;
+
+	// Gather input objects
+	TArray<UObject*> InputObjects;
+	InputObjects.Append(ToolBuilderUtil::FindAllComponents(SceneState, [this](UActorComponent* Component)
 	{
-		return CanBuildTarget(Object, TargetType);
-	});
-	TArray<TObjectPtr<UToolTarget>> Targets;
-	Targets.Reserve(Components.Num());
-	for (UActorComponent* Component : Components)
+		return Component && (!InputFilterFunction || InputFilterFunction(Component));
+	}));
+
+	TArray<bool> Used;
+	for (TObjectPtr<UToolTargetFactory> Factory : Factories)
 	{
-		Targets.Add(BuildTarget(Component, TargetType));
+		int32 NumTargets = Factory->CanBuildTargets(InputObjects, TargetType, Used);
+		if (NumTargets > 0)
+		{
+			TargetsOut.Append(Factory->BuildTargets(InputObjects, TargetType, Used));
+			ToolTargetManagerLocals::RemoveByFlags(InputObjects, Used);
+
+			if (InputObjects.IsEmpty())
+			{
+				break;
+			}
+		}
 	}
-	return Targets;
+
+	return TargetsOut;
 }
 
diff --git a/Engine/Source/Runtime/InteractiveToolsFramework/Private/ToolTargets/ToolTarget.cpp b/Engine/Source/Runtime/InteractiveToolsFramework/Private/ToolTargets/ToolTarget.cpp
--- a/Engine/Source/Runtime/InteractiveToolsFramework/Private/ToolTargets/ToolTarget.cpp
+++ b/Engine/Source/Runtime/InteractiveToolsFramework/Private/ToolTargets/ToolTarget.cpp
@@ -21,3 +21,49 @@ bool FToolTargetTypeRequirements::AreSatisfiedBy(UToolTarget* ToolTarget) const
 {
 	return !ToolTarget ? false : AreSatisfiedBy(ToolTarget->GetClass());
 }
+
+int32 UToolTargetFactory::CanBuildTargets(const TArray<UObject*>& InputObjects, const FToolTargetTypeRequirements& TargetTypeInfo, TArray<bool>& WouldBeUsedOut)
+{
+	int32 Count = 0;
+	WouldBeUsedOut.SetNum(InputObjects.Num());
+	for (int32 ObjIndex = 0; ObjIndex < InputObjects.Num(); ++ObjIndex)
+	{
+		UObject* Object = InputObjects[ObjIndex];
+		WouldBeUsedOut[ObjIndex] = CanBuildTarget(Object, TargetTypeInfo);
+		Count += WouldBeUsedOut[ObjIndex];
+	}
+	return Count;
+}
+
+TArray<UToolTarget*> UToolTargetFactory::BuildTargets(const TArray<UObject*>& InputObjects, const FToolTargetTypeRequirements& TargetTypeInfo, TArray<bool>& WasUsedOut)
+{
+	TArray<UToolTarget*> Output;
+	WasUsedOut.SetNum(InputObjects.Num());
+	for (int32 ObjIndex = 0; ObjIndex < InputObjects.Num(); ++ObjIndex)
+	{
+		UObject* Object = InputObjects[ObjIndex];
+		WasUsedOut[ObjIndex] = CanBuildTarget(Object, TargetTypeInfo);
+		if (WasUsedOut[ObjIndex])
+		{
+			Output.Add(BuildTarget(Object, TargetTypeInfo));
+		}
+	}
+
+	return Output;
+}
+
+UToolTarget* UToolTargetFactory::BuildFirstTarget(const TArray<UObject*>& InputObjects, const FToolTargetTypeRequirements& TargetTypeInfo, TArray<bool>& WasUsedOut)
+{
+	WasUsedOut.SetNumZeroed(InputObjects.Num());
+	for (int32 ObjIndex = 0; ObjIndex < InputObjects.Num(); ++ObjIndex)
+	{
+		UObject* Object = InputObjects[ObjIndex];
+		WasUsedOut[ObjIndex] = CanBuildTarget(Object, TargetTypeInfo);
+		if (WasUsedOut[ObjIndex])
+		{
+			return BuildTarget(Object, TargetTypeInfo);
+		}
+	}
+
+	return nullptr;
+}
diff --git a/Engine/Source/Runtime/InteractiveToolsFramework/Public/ToolTargets/ToolTarget.h b/Engine/Source/Runtime/InteractiveToolsFramework/Public/ToolTargets/ToolTarget.h
--- a/Engine/Source/Runtime/InteractiveToolsFramework/Public/ToolTargets/ToolTarget.h
+++ b/Engine/Source/Runtime/InteractiveToolsFramework/Public/ToolTargets/ToolTarget.h
@@ -93,8 +93,37 @@ class UToolTargetFactory : public UObject
 	GENERATED_BODY()
 
 public:
-
 	virtual bool CanBuildTarget(UObject* SourceObject, const FToolTargetTypeRequirements& TargetTypeInfo) const PURE_VIRTUAL(UToolTargetFactory::CanBuildTarget, return false;);
-
 	virtual UToolTarget* BuildTarget(UObject* SourceObject, const FToolTargetTypeRequirements& TargetTypeInfo) PURE_VIRTUAL(UToolTargetFactory::BuildTarget, return nullptr;);
+
+
+	// The following methods are used when building from multiple input objects. If not overriden, they will
+	//  call BuildTarget/CanBuildTarget on each object in sequence, but they can be overriden to create
+	//  combined tool targets out of multiple inputs.
+	
+	/**
+	 * Called to see if a factory could make one or more targets out of the given array of input objects, 
+	 *  if multiple are provided.
+	 * 
+	 * @param WouldBeUsedOut Set to be 1:1 with InputObjects, with each entry set to true if that input object
+	 *   would be used in the creation of a tool target
+	 * @return Number of targets that would be returned by the equivalent BuildTargets call.
+	 */
+	INTERACTIVETOOLSFRAMEWORK_API virtual int32 CanBuildTargets(const TArray<UObject*>& InputObjects, const FToolTargetTypeRequirements& TargetTypeInfo, TArray<bool>& WouldBeUsedOut);
+	
+	/**
+	 * Called to build one or more targets out of the given array of input objects, if multiple are provided.
+	 * 
+	 * @param WouldBeUsedOut Set to be 1:1 with InputObjects, with each entry set to true if that input object
+	 *   was used in the creation of a tool target
+	 */
+	INTERACTIVETOOLSFRAMEWORK_API virtual TArray<UToolTarget*> BuildTargets(const TArray<UObject*>& InputObjects, const FToolTargetTypeRequirements& TargetTypeInfo, TArray<bool>& WasUsedOut);
+
+	/**
+	 * Called by BuildFirstSelectedTargetable, intended to create just one target out of the given input.
+	 * 
+	 * @param WouldBeUsedOut Set to be 1:1 with InputObjects, with each entry set to true if that input object
+	 *   was used in the creation of a tool target
+	 */
+	INTERACTIVETOOLSFRAMEWORK_API virtual UToolTarget* BuildFirstTarget(const TArray<UObject*>& InputObjects, const FToolTargetTypeRequirements& TargetTypeInfo, TArray<bool>& WasUsedOut);
 };
