InteractiveToolsFramework: Catch errors where a tool target claims to be able to build a target, but doesn't claim any inputs.

diff --git a/Engine/Source/Runtime/InteractiveToolsFramework/Private/ToolTargetManager.cpp b/Engine/Source/Runtime/InteractiveToolsFramework/Private/ToolTargetManager.cpp
--- a/Engine/Source/Runtime/InteractiveToolsFramework/Private/ToolTargetManager.cpp
+++ b/Engine/Source/Runtime/InteractiveToolsFramework/Private/ToolTargetManager.cpp
@@ -9,12 +9,13 @@
 
 namespace ToolTargetManagerLocals
 {
-	// Given an array and a 1:1 array of bools, remove all entries that have the corresponding bool set to true
-	void RemoveByFlags(TArray<UObject*>& ObjectsInOut, const TArray<bool>& RemoveFlags)
+	// Given an array and a 1:1 array of bools, remove all entries that have the corresponding bool set to true.
+	// Returns number of elements removed
+	int32 RemoveByFlags(TArray<UObject*>& ObjectsInOut, const TArray<bool>& RemoveFlags)
 	{
 		if (!ensure(RemoveFlags.Num() == ObjectsInOut.Num()))
 		{
-			return;
+			return 0;
 		}
 
 		int32 OutIndex = 0;
@@ -29,6 +30,7 @@ namespace ToolTargetManagerLocals
 			++InIndex;
 		}
 		ObjectsInOut.SetNum(OutIndex);
+		return InIndex - OutIndex;
 	}
 }
 
@@ -134,7 +136,8 @@ void UToolTargetManager::EnumerateSelectedAndTargetableComponents(const FToolBui
 					ComponentFunc(static_cast<UActorComponent*>(InputObjects[Index]));
 				}
 			}
-			ToolTargetManagerLocals::RemoveByFlags(InputObjects, Used);
+			ensureMsgf(ToolTargetManagerLocals::RemoveByFlags(InputObjects, Used) > 0, 
+				TEXT("Factory claimed it could build target(s) without using any objects."));
 
 			if (InputObjects.IsEmpty())
 			{
@@ -165,7 +168,8 @@ int32 UToolTargetManager::CountSelectedAndTargetableWithPredicate(const FToolBui
 		if (NumTargets > 0)
 		{
 			Count += NumTargets;
-			ToolTargetManagerLocals::RemoveByFlags(InputObjects, Used);
+			ensureMsgf(ToolTargetManagerLocals::RemoveByFlags(InputObjects, Used) > 0,
+				TEXT("Factory claimed it could build target(s) without using any objects."));
 			if (InputObjects.IsEmpty())
 			{
 				break;
@@ -216,7 +220,8 @@ TArray<TObjectPtr<UToolTarget>> UToolTargetManager::BuildAllSelectedTargetable(c
 		if (NumTargets > 0)
 		{
 			TargetsOut.Append(Factory->BuildTargets(InputObjects, TargetType, Used));
-			ToolTargetManagerLocals::RemoveByFlags(InputObjects, Used);
+			ensureMsgf(ToolTargetManagerLocals::RemoveByFlags(InputObjects, Used) > 0,
+				TEXT("Factory claimed it could build target(s) without using any objects."));
 
 			if (InputObjects.IsEmpty())
 			{
