InteractiveToolManager: Add ability to query for the current change tracking mode

diff --git a/Engine/Source/Runtime/InteractiveToolsFramework/Private/InteractiveToolManager.cpp b/Engine/Source/Runtime/InteractiveToolsFramework/Private/InteractiveToolManager.cpp
--- a/Engine/Source/Runtime/InteractiveToolsFramework/Private/InteractiveToolManager.cpp
+++ b/Engine/Source/Runtime/InteractiveToolsFramework/Private/InteractiveToolManager.cpp
@@ -377,6 +377,10 @@ void UInteractiveToolManager::ConfigureChangeTrackingMode(EToolChangeTrackingMod
 	ActiveToolChangeTrackingMode = ChangeMode;
 }
 
+EToolChangeTrackingMode UInteractiveToolManager::GetChangeTrackingMode() const
+{
+	return ActiveToolChangeTrackingMode;
+}
 
 
 void UInteractiveToolManager::Tick(float DeltaTime)
diff --git a/Engine/Source/Runtime/InteractiveToolsFramework/Public/InteractiveToolManager.h b/Engine/Source/Runtime/InteractiveToolsFramework/Public/InteractiveToolManager.h
--- a/Engine/Source/Runtime/InteractiveToolsFramework/Public/InteractiveToolManager.h
+++ b/Engine/Source/Runtime/InteractiveToolsFramework/Public/InteractiveToolManager.h
@@ -229,6 +229,10 @@ public:
 	 */
 	INTERACTIVETOOLSFRAMEWORK_API virtual void ConfigureChangeTrackingMode(EToolChangeTrackingMode ChangeMode);
 
+	/**
+	 * Returns how tool changes emit change events. See EToolChangeTrackingMode for details.
+	 */
+	INTERACTIVETOOLSFRAMEWORK_API virtual EToolChangeTrackingMode GetChangeTrackingMode() const;
 
 	//
 	// Functions that Tools can call to interact with Transactions API
