ITF Tool Targets: Fix Primitive/SceneComponentToolTarget object replacement logic

diff --git a/Engine/Source/Runtime/InteractiveToolsFramework/Private/ToolTargets/PrimitiveComponentToolTarget.cpp b/Engine/Source/Runtime/InteractiveToolsFramework/Private/ToolTargets/PrimitiveComponentToolTarget.cpp
--- a/Engine/Source/Runtime/InteractiveToolsFramework/Private/ToolTargets/PrimitiveComponentToolTarget.cpp
+++ b/Engine/Source/Runtime/InteractiveToolsFramework/Private/ToolTargets/PrimitiveComponentToolTarget.cpp
@@ -84,7 +84,7 @@ void UPrimitiveComponentToolTarget::OnObjectsReplaced(const FCoreUObjectDelegate
 	// while the tool is running, we are unable to unhide it later. The user is unlikely to understand 
 	// why their object disappeared in that case or know to fix it via the component visibility property.
 
-	UObject* const * MappedObject = Map.Find(Component.Get());
+	UObject* const * MappedObject = Map.Find(Component.Get(/* bEvenIfPendingKill = */ true));
 	if (MappedObject)
 	{
 		Component = Cast<UPrimitiveComponent>(*MappedObject);
diff --git a/Engine/Source/Runtime/InteractiveToolsFramework/Private/ToolTargets/SceneComponentToolTarget.cpp b/Engine/Source/Runtime/InteractiveToolsFramework/Private/ToolTargets/SceneComponentToolTarget.cpp
--- a/Engine/Source/Runtime/InteractiveToolsFramework/Private/ToolTargets/SceneComponentToolTarget.cpp
+++ b/Engine/Source/Runtime/InteractiveToolsFramework/Private/ToolTargets/SceneComponentToolTarget.cpp
@@ -69,7 +69,7 @@ void USceneComponentToolTarget::OnObjectsReplaced(const FCoreUObjectDelegates::F
 	// while the tool is running, we are unable to unhide it later. The user is unlikely to understand 
 	// why their object disappeared in that case or know to fix it via the component visibility property.
 
-	UObject* const * MappedObject = Map.Find(Component.Get());
+	UObject* const * MappedObject = Map.Find(Component.Get(/* bEvenIfPendingKill = */ true));
 	if (MappedObject)
 	{
 		Component = Cast<USceneComponent>(*MappedObject);
