ToolTargetManager: Add RemoveTargetFactoriesByPredicate

diff --git a/Engine/Source/Runtime/InteractiveToolsFramework/Private/ToolTargetManager.cpp b/Engine/Source/Runtime/InteractiveToolsFramework/Private/ToolTargetManager.cpp
--- a/Engine/Source/Runtime/InteractiveToolsFramework/Private/ToolTargetManager.cpp
+++ b/Engine/Source/Runtime/InteractiveToolsFramework/Private/ToolTargetManager.cpp
@@ -59,6 +59,10 @@ void UToolTargetManager::AddTargetFactory(UToolTargetFactory* Factory)
 	Factories.Add(Factory);
 }
 
+void UToolTargetManager::RemoveTargetFactoriesByPredicate(TFunctionRef<bool(UToolTargetFactory*)> Predicate)
+{
+	Factories.RemoveAll(Predicate);
+}
 
 UToolTargetFactory* UToolTargetManager::FindFirstFactoryByPredicate(TFunctionRef<bool(UToolTargetFactory*)> Predicate)
 {
diff --git a/Engine/Source/Runtime/InteractiveToolsFramework/Public/ToolTargetManager.h b/Engine/Source/Runtime/InteractiveToolsFramework/Public/ToolTargetManager.h
--- a/Engine/Source/Runtime/InteractiveToolsFramework/Public/ToolTargetManager.h
+++ b/Engine/Source/Runtime/InteractiveToolsFramework/Public/ToolTargetManager.h
@@ -48,6 +48,11 @@ public:
 	 */
 	INTERACTIVETOOLSFRAMEWORK_API virtual void AddTargetFactory(UToolTargetFactory* Factory);
 
+	/**
+	 * Remove existing ToolTargetFactories by predicate
+	 */
+	INTERACTIVETOOLSFRAMEWORK_API virtual void RemoveTargetFactoriesByPredicate(TFunctionRef<bool(UToolTargetFactory*)> Predicate);
+
 	/**
 	 * Find the first ToolTargetFactory that passes the Predicate function
 	 */
