[ITF] Added a HitPriority to Gizmo Elements, where if elements overlap those with the higher priority are selected over those with a lower priority. Defaults to 0.

diff --git a/Engine/Source/Editor/Experimental/EditorInteractiveToolsFramework/Private/EditorGizmos/GizmoElementGimbal.cpp b/Engine/Source/Editor/Experimental/EditorInteractiveToolsFramework/Private/EditorGizmos/GizmoElementGimbal.cpp
--- a/Engine/Source/Editor/Experimental/EditorInteractiveToolsFramework/Private/EditorGizmos/GizmoElementGimbal.cpp
+++ b/Engine/Source/Editor/Experimental/EditorInteractiveToolsFramework/Private/EditorGizmos/GizmoElementGimbal.cpp
@@ -38,8 +38,7 @@ void UGizmoElementGimbal::Render(IToolsContextRenderAPI* InRenderAPI, const FRen
 	}
 }
 
-FInputRayHit UGizmoElementGimbal::LineTrace(const UGizmoViewContext* InViewContext,
-	const FLineTraceTraversalState& InLineTraceState, const FVector& InRayOrigin, const FVector& InRayDirection)
+FInputRayHit UGizmoElementGimbal::LineTrace(const UGizmoViewContext* ViewContext, const FLineTraceTraversalState& LineTraceState, const FVector& RayOrigin, const FVector& RayDirection, FLineTraceOutput& OutLineTraceOutput)
 {
 	using namespace UE::GizmoRotationUtil;
 	
@@ -49,8 +48,8 @@ FInputRayHit UGizmoElementGimbal::LineTrace(const UGizmoViewContext* InViewConte
 		return Hit;
 	}
 
-	FLineTraceTraversalState CurrentLineTraceState(InLineTraceState);
-	const bool bHittable = UpdateLineTraceState(InViewContext, FVector::ZeroVector, CurrentLineTraceState);
+	FLineTraceTraversalState CurrentLineTraceState(LineTraceState);
+	const bool bHittable = UpdateLineTraceState(ViewContext, FVector::ZeroVector, CurrentLineTraceState);
 
 	if (!bHittable)
 	{
@@ -67,7 +66,8 @@ FInputRayHit UGizmoElementGimbal::LineTrace(const UGizmoViewContext* InViewConte
 		CurrentLineTraceState.LocalToWorldTransform.SetRotation(Decomposition.R[Index++]);
 		if (Element)
 		{
-			FInputRayHit NewHit = Element->LineTrace(InViewContext, CurrentLineTraceState, InRayOrigin, InRayDirection);
+			FLineTraceOutput NewLineTraceOutput;
+			FInputRayHit NewHit = Element->LineTrace(ViewContext, CurrentLineTraceState, RayOrigin, RayDirection, NewLineTraceOutput);
 			if (!Hit.bHit || NewHit.HitDepth < Hit.HitDepth)
 			{
 				Hit = NewHit;
diff --git a/Engine/Source/Editor/Experimental/EditorInteractiveToolsFramework/Public/EditorGizmos/GizmoElementGimbal.h b/Engine/Source/Editor/Experimental/EditorInteractiveToolsFramework/Public/EditorGizmos/GizmoElementGimbal.h
--- a/Engine/Source/Editor/Experimental/EditorInteractiveToolsFramework/Public/EditorGizmos/GizmoElementGimbal.h
+++ b/Engine/Source/Editor/Experimental/EditorInteractiveToolsFramework/Public/EditorGizmos/GizmoElementGimbal.h
@@ -22,7 +22,7 @@ public:
 
 	//~ Begin UGizmoElementBase Interface.
 	UE_API virtual void Render(IToolsContextRenderAPI* InRenderAPI, const FRenderTraversalState& InRenderState) override;
-	UE_API virtual FInputRayHit LineTrace(const UGizmoViewContext* InViewContext, const FLineTraceTraversalState& InLineTraceState, const FVector& InRayOrigin, const FVector& InRayDirection) override;
+	UE_API virtual FInputRayHit LineTrace(const UGizmoViewContext* ViewContext, const FLineTraceTraversalState& LineTraceState, const FVector& RayOrigin, const FVector& RayDirection, FLineTraceOutput& OutLineTraceOutput) override;
 	//~ End UGizmoElementBase Interface.
 
 	// Add object to group.
diff --git a/Engine/Source/Runtime/InteractiveToolsFramework/Private/BaseGizmos/GizmoElementArc.cpp b/Engine/Source/Runtime/InteractiveToolsFramework/Private/BaseGizmos/GizmoElementArc.cpp
--- a/Engine/Source/Runtime/InteractiveToolsFramework/Private/BaseGizmos/GizmoElementArc.cpp
+++ b/Engine/Source/Runtime/InteractiveToolsFramework/Private/BaseGizmos/GizmoElementArc.cpp
@@ -183,14 +183,14 @@ void UGizmoElementArc::Render(IToolsContextRenderAPI* RenderAPI, const FRenderTr
 		Start, End, NumSegments, UseMaterial->GetRenderProxy(), VertexColor, LineColor, CurrentLineThickness);
 }
 
-FInputRayHit UGizmoElementArc::LineTrace(const UGizmoViewContext* ViewContext, const FLineTraceTraversalState& LineTraceState, const FVector& RayOrigin, const FVector& RayDirection)
+FInputRayHit UGizmoElementArc::LineTrace(const UGizmoViewContext* ViewContext, const FLineTraceTraversalState& LineTraceState, const FVector& RayOrigin, const FVector& RayDirection, FLineTraceOutput& OutLineTraceOutput)
 {
 	FLineTraceTraversalState CurrentLineTraceState(LineTraceState);
 	bool bHittableViewDependent = UpdateLineTraceState(ViewContext, Center, CurrentLineTraceState);
 
 	if (bHittableViewDependent)
 	{
-		const double PixelHitThresholdAdjust = CurrentLineTraceState.PixelToWorldScale * PixelHitDistanceThreshold;
+		double PixelHitThresholdAdjust = CurrentLineTraceState.PixelToWorldScale * PixelHitDistanceThreshold;
 		const FVector WorldCenter = CurrentLineTraceState.LocalToWorldTransform.TransformPosition(FVector::ZeroVector);
 		const FVector WorldAxis0 = CurrentLineTraceState.LocalToWorldTransform.TransformVectorNoScale(Axis0);
 		const FVector WorldAxis1 = CurrentLineTraceState.LocalToWorldTransform.TransformVectorNoScale(Axis1);
@@ -200,7 +200,7 @@ FInputRayHit UGizmoElementArc::LineTrace(const UGizmoViewContext* ViewContext, c
 		bool bPartial = IsPartial(ViewContext, WorldCenter, WorldNormal);
 		const double PartialAngle = PartialEndAngle - PartialStartAngle;
 
-		if (bPartial && PartialAngle <= 0.0)
+		if (bPartial && PartialAngle <= UE_DOUBLE_SMALL_NUMBER)
 		{
 			return FInputRayHit();
 		}
@@ -208,7 +208,7 @@ FInputRayHit UGizmoElementArc::LineTrace(const UGizmoViewContext* ViewContext, c
 		// Intersect ray with plane in which arc lies.
 		FPlane Plane(WorldCenter, WorldNormal);
 		HitDepth = FMath::RayPlaneIntersectionParam(RayOrigin, RayDirection, Plane);
-		if (HitDepth < 0.0)
+		if (HitDepth < UE_DOUBLE_SMALL_NUMBER)
 		{
 			return FInputRayHit();
 		}
@@ -220,8 +220,8 @@ FInputRayHit UGizmoElementArc::LineTrace(const UGizmoViewContext* ViewContext, c
 		const double WorldOuterRadius = Radius * CurrentLineTraceState.LocalToWorldTransform.GetScale3D().X;
 		const double WorldInnerRadius = InnerRadius * CurrentLineTraceState.LocalToWorldTransform.GetScale3D().X;
 		const double Distance = HitVec.Length();
-		const double HitBufferMax = WorldOuterRadius + PixelHitThresholdAdjust;
-		const double HitBufferMin = WorldInnerRadius - PixelHitThresholdAdjust;
+		double HitBufferMax = WorldOuterRadius + PixelHitThresholdAdjust;
+		double HitBufferMin = WorldInnerRadius - PixelHitThresholdAdjust;
 
 		if (Distance > HitBufferMax || Distance < HitBufferMin)
 		{
@@ -233,10 +233,9 @@ FInputRayHit UGizmoElementArc::LineTrace(const UGizmoViewContext* ViewContext, c
 		{
 			// Compute projected angle
 			FVector WorldBeginAxis = WorldAxis0.RotateAngleAxis(PartialStartAngle, WorldNormal).GetSafeNormal();
-			FVector HitVecNormal = HitVec.GetSafeNormal();
 			double HitAngle = UE::Geometry::VectorUtil::PlaneAngleSignedR(WorldBeginAxis, HitVec, WorldNormal); 
 
-			if (HitAngle < 0)
+			if (HitAngle < UE_DOUBLE_SMALL_NUMBER)
 			{
 				HitAngle = UE_DOUBLE_TWO_PI + HitAngle;
 			}
@@ -247,11 +246,17 @@ FInputRayHit UGizmoElementArc::LineTrace(const UGizmoViewContext* ViewContext, c
 			}
 		}
 
-		if (HitDepth >= 0.0)
+		if (HitDepth >= UE_DOUBLE_SMALL_NUMBER)
 		{
-			FInputRayHit RayHit(static_cast<float>(HitDepth));
-			RayHit.SetHitObject(this);
-			RayHit.HitIdentifier = PartIdentifier;
+			FInputRayHit RayHit = MakeRayHit(HitDepth, OutLineTraceOutput);
+
+			// We successfully hit, now check if we intersect the actual surface (rather than just within the PixelHitDistanceThreshold) ...
+			PixelHitThresholdAdjust = CurrentLineTraceState.PixelToWorldScale * FMath::Max(0.0, MinimumPixelHitDistanceThreshold);
+			HitBufferMax = WorldOuterRadius + PixelHitThresholdAdjust;
+			HitBufferMin = WorldInnerRadius - PixelHitThresholdAdjust;
+
+			OutLineTraceOutput.bIsSurfaceHit = Distance <= HitBufferMax && Distance >= HitBufferMin;
+
 			return RayHit;
 		}
 	}
diff --git a/Engine/Source/Runtime/InteractiveToolsFramework/Private/BaseGizmos/GizmoElementArrow.cpp b/Engine/Source/Runtime/InteractiveToolsFramework/Private/BaseGizmos/GizmoElementArrow.cpp
--- a/Engine/Source/Runtime/InteractiveToolsFramework/Private/BaseGizmos/GizmoElementArrow.cpp
+++ b/Engine/Source/Runtime/InteractiveToolsFramework/Private/BaseGizmos/GizmoElementArrow.cpp
@@ -24,6 +24,14 @@ void UpdatePixelThreshold(UGizmoElementBase* InGizmoElement, const float InPixel
 	}		
 }
 
+void UpdateMinPixelThreshold(UGizmoElementBase* InGizmoElement, const float InPixelHitThreshold)
+{
+	if (InGizmoElement)
+	{
+		InGizmoElement->SetMinimumPixelHitDistanceThreshold(InPixelHitThreshold);
+	}		
+}
+
 using ThresholdGuard = TGuardValue_Bitfield_Cleanup<TFunction<void()>>;
 
 }
@@ -69,7 +77,7 @@ void UGizmoElementArrow::Render(IToolsContextRenderAPI* RenderAPI, const FRender
 	}
 }
 
-FInputRayHit UGizmoElementArrow::LineTrace(const UGizmoViewContext* ViewContext, const FLineTraceTraversalState& LineTraceState, const FVector& RayOrigin, const FVector& RayDirection)
+FInputRayHit UGizmoElementArrow::LineTrace(const UGizmoViewContext* ViewContext, const FLineTraceTraversalState& LineTraceState, const FVector& RayOrigin, const FVector& RayDirection, FLineTraceOutput& OutLineTraceOutput)
 {
 	using namespace UE::Private;
 
@@ -85,12 +93,15 @@ FInputRayHit UGizmoElementArrow::LineTrace(const UGizmoViewContext* ViewContext,
 	{
 		UpdatePixelThreshold(CylinderElement, PixelHitDistanceThreshold);
 		UpdatePixelThreshold(HeadElement, PixelHitDistanceThreshold);
+
+		UpdateMinPixelThreshold(CylinderElement, MinimumPixelHitDistanceThreshold);
+		UpdateMinPixelThreshold(HeadElement, MinimumPixelHitDistanceThreshold);
 	});
 
 	// update threshold if hitting the mask
 	if (UGizmoElementBase* HitMaskGizmo = HitMask.Get())
 	{
-		const FInputRayHit MaskHit = HitMaskGizmo->LineTrace(ViewContext, LineTraceState, RayOrigin, RayDirection);
+		const FInputRayHit MaskHit = HitMaskGizmo->LineTrace(ViewContext, LineTraceState, RayOrigin, RayDirection, OutLineTraceOutput);
 		if (MaskHit.bHit)
 		{
 			static constexpr float NoPixelHitThreshold = 0.f;
@@ -100,12 +111,12 @@ FInputRayHit UGizmoElementArrow::LineTrace(const UGizmoViewContext* ViewContext,
 	}
 	
 	check(CylinderElement);
-	FInputRayHit Hit = CylinderElement->LineTrace(ViewContext, CurrentLineTraceState, RayOrigin, RayDirection);
+	FInputRayHit Hit = CylinderElement->LineTrace(ViewContext, CurrentLineTraceState, RayOrigin, RayDirection, OutLineTraceOutput);
 
 	if (!Hit.bHit)
 	{
 		check(HeadElement);
-		Hit = HeadElement->LineTrace(ViewContext, CurrentLineTraceState, RayOrigin, RayDirection);
+		Hit = HeadElement->LineTrace(ViewContext, CurrentLineTraceState, RayOrigin, RayDirection, OutLineTraceOutput);
 	}
 
 	if (Hit.bHit)
@@ -254,6 +265,16 @@ void UGizmoElementArrow::SetPixelHitDistanceThreshold(float InPixelHitDistanceTh
 	}
 }
 
+void UGizmoElementArrow::SetMinimumPixelHitDistanceThreshold(float InMinimumPixelHitDistanceThreshold)
+{
+	if (MinimumPixelHitDistanceThreshold != InMinimumPixelHitDistanceThreshold)
+	{
+		MinimumPixelHitDistanceThreshold = InMinimumPixelHitDistanceThreshold;
+		bUpdateArrowBody = true;
+		bUpdateArrowHead = true;
+	}
+}
+
 void UGizmoElementArrow::SetHitMask(const TWeakObjectPtr<UGizmoElementBase>& InHitMask)
 {
 	HitMask = InHitMask;
diff --git a/Engine/Source/Runtime/InteractiveToolsFramework/Private/BaseGizmos/GizmoElementArrowHead.cpp b/Engine/Source/Runtime/InteractiveToolsFramework/Private/BaseGizmos/GizmoElementArrowHead.cpp
--- a/Engine/Source/Runtime/InteractiveToolsFramework/Private/BaseGizmos/GizmoElementArrowHead.cpp
+++ b/Engine/Source/Runtime/InteractiveToolsFramework/Private/BaseGizmos/GizmoElementArrowHead.cpp
@@ -71,7 +71,7 @@ void UGizmoElementArrowHead::Render(IToolsContextRenderAPI* RenderAPI, const FRe
 	}
 }
 
-FInputRayHit UGizmoElementArrowHead::LineTrace(const UGizmoViewContext* ViewContext, const FLineTraceTraversalState& LineTraceState, const FVector& RayOrigin, const FVector& RayDirection)
+FInputRayHit UGizmoElementArrowHead::LineTrace(const UGizmoViewContext* ViewContext, const FLineTraceTraversalState& LineTraceState, const FVector& RayOrigin, const FVector& RayDirection, FLineTraceOutput& OutLineTraceOutput)
 {
 	using namespace UE::InteractiveToolsFramework::Private::GizmoElementArrowHeadLocals;
 
@@ -93,7 +93,7 @@ FInputRayHit UGizmoElementArrowHead::LineTrace(const UGizmoViewContext* ViewCont
 	// update threshold if hitting the mask
 	if (UGizmoElementBase* HitMaskGizmo = HitMask.Get())
 	{
-		const FInputRayHit MaskHit = HitMaskGizmo->LineTrace(ViewContext, LineTraceState, RayOrigin, RayDirection);
+		const FInputRayHit MaskHit = HitMaskGizmo->LineTrace(ViewContext, LineTraceState, RayOrigin, RayDirection, OutLineTraceOutput);
 		if (MaskHit.bHit)
 		{
 			static constexpr float NoPixelHitThreshold = 0.0f;
@@ -106,7 +106,7 @@ FInputRayHit UGizmoElementArrowHead::LineTrace(const UGizmoViewContext* ViewCont
 	auto LineTraceElement = [&](UGizmoElementBase* InElement) -> FInputRayHit
 	{
 		check(InElement);
-		return InElement->LineTrace(ViewContext, CurrentLineTraceState, RayOrigin, RayDirection);
+		return InElement->LineTrace(ViewContext, CurrentLineTraceState, RayOrigin, RayDirection, OutLineTraceOutput);
 	};
 
 	FInputRayHit Hit;
diff --git a/Engine/Source/Runtime/InteractiveToolsFramework/Private/BaseGizmos/GizmoElementBase.cpp b/Engine/Source/Runtime/InteractiveToolsFramework/Private/BaseGizmos/GizmoElementBase.cpp
--- a/Engine/Source/Runtime/InteractiveToolsFramework/Private/BaseGizmos/GizmoElementBase.cpp
+++ b/Engine/Source/Runtime/InteractiveToolsFramework/Private/BaseGizmos/GizmoElementBase.cpp
@@ -476,6 +476,17 @@ bool UGizmoElementBase::UpdateLineTraceState(const UGizmoViewContext* ViewContex
 	return false;
 }
 
+FInputRayHit UGizmoElementBase::MakeRayHit(const double InHitDepth, FLineTraceOutput& OutLineTraceOutput)
+{
+	FInputRayHit Hit(static_cast<float>(InHitDepth));
+	Hit.SetHitObject(this);
+	Hit.HitIdentifier = PartIdentifier;
+
+	OutLineTraceOutput.HitPriority = HitPriority;
+
+	return Hit;
+}
+
 bool UGizmoElementBase::GetVisibleState() const
 {
 	return (static_cast<uint8>(ElementState) & static_cast<uint8>(EGizmoElementState::Visible));
@@ -486,6 +497,12 @@ void UGizmoElementBase::DrawHUD(FCanvas* Canvas, IToolsContextRenderAPI* RenderA
 	// Empty default implementation
 }
 
+FInputRayHit UGizmoElementBase::LineTrace(const UGizmoViewContext* ViewContext, const FLineTraceTraversalState& LineTraceState, const FVector& RayOrigin, const FVector& RayDirection)
+{
+	FLineTraceOutput Unused;
+	return LineTrace(ViewContext, LineTraceState, RayOrigin, RayDirection, Unused);
+}
+
 void UGizmoElementBase::SetVisibleState(bool bVisible)
 {
 	ElementState = (bVisible ? ElementState | EGizmoElementState::Visible : ElementState & (~EGizmoElementState::Visible));
@@ -496,6 +513,16 @@ bool UGizmoElementBase::GetHittableState() const
 	return (static_cast<uint8>(ElementState) & static_cast<uint8>(EGizmoElementState::Hittable));
 }
 
+void UGizmoElementBase::SetHitPriority(const int32 InHitPriority)
+{
+	HitPriority = InHitPriority;
+}
+
+int32 UGizmoElementBase::GetHitPriority() const
+{
+	return HitPriority;
+}
+
 void UGizmoElementBase::SetHittableState(bool bHittable)
 {
 	ElementState = (bHittable ? ElementState | EGizmoElementState::Hittable : ElementState & (~EGizmoElementState::Hittable));
@@ -810,6 +837,16 @@ float UGizmoElementBase::GetPixelHitDistanceThreshold() const
 	return PixelHitDistanceThreshold;
 }
 
+void UGizmoElementBase::SetMinimumPixelHitDistanceThreshold(const float InMinimumPixelHitDistanceThreshold)
+{
+	MinimumPixelHitDistanceThreshold = InMinimumPixelHitDistanceThreshold;
+}
+
+float UGizmoElementBase::GetMinimumPixelHitDistanceThreshold() const
+{
+	return MinimumPixelHitDistanceThreshold;
+}
+
 void UGizmoElementBase::SetMaterial(TWeakObjectPtr<UMaterialInterface> InMaterial, bool InOverridesChildState)
 {
 	MeshRenderAttributes.Material.SetMaterial(InMaterial, InOverridesChildState);
diff --git a/Engine/Source/Runtime/InteractiveToolsFramework/Private/BaseGizmos/GizmoElementBox.cpp b/Engine/Source/Runtime/InteractiveToolsFramework/Private/BaseGizmos/GizmoElementBox.cpp
--- a/Engine/Source/Runtime/InteractiveToolsFramework/Private/BaseGizmos/GizmoElementBox.cpp
+++ b/Engine/Source/Runtime/InteractiveToolsFramework/Private/BaseGizmos/GizmoElementBox.cpp
@@ -27,7 +27,7 @@ void UGizmoElementBox::Render(IToolsContextRenderAPI* RenderAPI, const FRenderTr
 	}
 }
 
-FInputRayHit UGizmoElementBox::LineTrace(const UGizmoViewContext* ViewContext, const FLineTraceTraversalState& LineTraceState, const FVector& RayOrigin, const FVector& RayDirection)
+FInputRayHit UGizmoElementBox::LineTrace(const UGizmoViewContext* ViewContext, const FLineTraceTraversalState& LineTraceState, const FVector& RayOrigin, const FVector& RayDirection, FLineTraceOutput& OutLineTraceOutput)
 {
 	FLineTraceTraversalState CurrentLineTraceState(LineTraceState);
 	bool bHittableViewDependent = UpdateLineTraceState(ViewContext, Center, CurrentLineTraceState);
@@ -39,18 +39,26 @@ FInputRayHit UGizmoElementBox::LineTrace(const UGizmoViewContext* ViewContext, c
 		const FVector XAxis = FVector::CrossProduct(YAxis, ZAxis);
 		const FVector WorldCenter = CurrentLineTraceState.LocalToWorldTransform.TransformPosition(FVector::ZeroVector);
 		const double Scale = CurrentLineTraceState.LocalToWorldTransform.GetScale3D().X;
-		const double PixelHitThresholdAdjust = CurrentLineTraceState.PixelToWorldScale * PixelHitDistanceThreshold;
-		const FVector WorldExtent = Dimensions * Scale * 0.5 + FVector(PixelHitThresholdAdjust);
+		double PixelHitThresholdAdjust = CurrentLineTraceState.PixelToWorldScale * PixelHitDistanceThreshold;
+		const FVector WorldExtent = Dimensions * Scale * 0.5;
 
 		double HitDepth = 0.0;
 		UE::Geometry::TRay<double> Ray(RayOrigin, RayDirection);
 		UE::Geometry::TFrame3<double> Frame(WorldCenter, XAxis, YAxis, ZAxis);
-		UE::Geometry::TOrientedBox3<double> Box(Frame, WorldExtent);
+		UE::Geometry::TOrientedBox3<double> Box(Frame, WorldExtent + PixelHitThresholdAdjust);
 		if (UE::Geometry::TIntrRay3OrientedBox3<double>::FindIntersection(Ray, Box, HitDepth))
 		{
-			FInputRayHit RayHit(static_cast<float>(HitDepth));
-			RayHit.SetHitObject(this);
-			RayHit.HitIdentifier = PartIdentifier;
+			FInputRayHit RayHit = MakeRayHit(HitDepth, OutLineTraceOutput);
+
+			// We successfully hit, now check if we intersect the actual surface (rather than just within the PixelHitDistanceThreshold) ...
+			PixelHitThresholdAdjust = CurrentLineTraceState.PixelToWorldScale * MinimumPixelHitDistanceThreshold;
+			Box = UE::Geometry::TOrientedBox3<double>(Frame, WorldExtent + PixelHitThresholdAdjust); // Recreate the box without PixelHitDistanceThreshold
+			if (UE::Geometry::TIntrRay3OrientedBox3<double>::FindIntersection(Ray, Box, HitDepth))
+			{
+				RayHit.HitDepth = HitDepth;
+				OutLineTraceOutput.bIsSurfaceHit = true;
+			}
+
 			return RayHit;
 		}
 	}
diff --git a/Engine/Source/Runtime/InteractiveToolsFramework/Private/BaseGizmos/GizmoElementCircle.cpp b/Engine/Source/Runtime/InteractiveToolsFramework/Private/BaseGizmos/GizmoElementCircle.cpp
--- a/Engine/Source/Runtime/InteractiveToolsFramework/Private/BaseGizmos/GizmoElementCircle.cpp
+++ b/Engine/Source/Runtime/InteractiveToolsFramework/Private/BaseGizmos/GizmoElementCircle.cpp
@@ -124,7 +124,7 @@ void UGizmoElementCircle::Render(IToolsContextRenderAPI* RenderAPI, const FRende
 	}
 }
 
-FInputRayHit UGizmoElementCircle::LineTrace(const UGizmoViewContext* ViewContext, const FLineTraceTraversalState& LineTraceState, const FVector& RayOrigin, const FVector& RayDirection)
+FInputRayHit UGizmoElementCircle::LineTrace(const UGizmoViewContext* ViewContext, const FLineTraceTraversalState& LineTraceState, const FVector& RayOrigin, const FVector& RayDirection, FLineTraceOutput& OutLineTraceOutput)
 {
 	FLineTraceTraversalState CurrentLineTraceState(LineTraceState);
 	bool bHittableViewDependent = UpdateLineTraceState(ViewContext, Center, CurrentLineTraceState);
@@ -136,6 +136,7 @@ FInputRayHit UGizmoElementCircle::LineTrace(const UGizmoViewContext* ViewContext
 		const FVector WorldAxis1 = CurrentLineTraceState.LocalToWorldTransform.TransformVectorNoScale(Axis1);
 		const FVector WorldNormal = WorldAxis0 ^ WorldAxis1;
 		const double PixelHitThresholdAdjust = CurrentLineTraceState.PixelToWorldScale * PixelHitDistanceThreshold;
+		const double MinimumPixelHitThresholdAdjust = CurrentLineTraceState.PixelToWorldScale * MinimumPixelHitDistanceThreshold;
 		double WorldRadius = CurrentLineTraceState.LocalToWorldTransform.GetScale3D().X * Radius;
 
 		// if ray is parallel to circle, no hit
@@ -153,9 +154,14 @@ FInputRayHit UGizmoElementCircle::LineTrace(const UGizmoViewContext* ViewContext
 
 			if (Result.intersects)
 			{
-				FInputRayHit RayHit(static_cast<float>(Result.parameter.Min));
-				RayHit.SetHitObject(this);
-				RayHit.HitIdentifier = PartIdentifier;
+				FInputRayHit RayHit = MakeRayHit(Result.parameter.Min, OutLineTraceOutput);
+
+				// We successfully hit, now check if we intersect the actual surface (rather than just within the PixelHitDistanceThreshold)
+				WorldRadius -= PixelHitThresholdAdjust;
+				WorldRadius += MinimumPixelHitThresholdAdjust;
+				IntersectionUtil::RayCircleIntersection(RayOrigin, RayDirection, WorldCenter, WorldRadius, WorldNormal, Result);
+				OutLineTraceOutput.bIsSurfaceHit = Result.intersects;
+
 				return RayHit;
 			}
 		}
@@ -163,7 +169,7 @@ FInputRayHit UGizmoElementCircle::LineTrace(const UGizmoViewContext* ViewContext
 		{
 			FPlane Plane(WorldCenter, WorldNormal);
 			double HitDepth = FMath::RayPlaneIntersectionParam(RayOrigin, RayDirection, Plane);
-			if (HitDepth < 0)
+			if (HitDepth < UE_DOUBLE_SMALL_NUMBER)
 			{
 				return FInputRayHit();
 			}
@@ -181,13 +187,16 @@ FInputRayHit UGizmoElementCircle::LineTrace(const UGizmoViewContext* ViewContext
 			
 			if (Distance <= HitBuffer)
 			{
-				FInputRayHit RayHit(static_cast<float>(HitDepth));
-				RayHit.SetHitObject(this);
-				RayHit.HitIdentifier = PartIdentifier;
+				FInputRayHit RayHit = MakeRayHit(HitDepth, OutLineTraceOutput);
+
+				// We successfully hit, now check if we intersect the actual surface (rather than just within the PixelHitDistanceThreshold)
+				OutLineTraceOutput.bIsSurfaceHit = Distance < FMath::Max(CurrentLineTraceState.PixelToWorldScale, MinimumPixelHitThresholdAdjust + LineThickness);
+
 				return RayHit;
 			}
 		}
 	}
+
 	return FInputRayHit();
 }
 
diff --git a/Engine/Source/Runtime/InteractiveToolsFramework/Private/BaseGizmos/GizmoElementCone.cpp b/Engine/Source/Runtime/InteractiveToolsFramework/Private/BaseGizmos/GizmoElementCone.cpp
--- a/Engine/Source/Runtime/InteractiveToolsFramework/Private/BaseGizmos/GizmoElementCone.cpp
+++ b/Engine/Source/Runtime/InteractiveToolsFramework/Private/BaseGizmos/GizmoElementCone.cpp
@@ -95,7 +95,7 @@ void UGizmoElementCone::Render(IToolsContextRenderAPI* RenderAPI, const FRenderT
 	}
 }
 
-FInputRayHit UGizmoElementCone::LineTrace(const UGizmoViewContext* ViewContext, const FLineTraceTraversalState& LineTraceState, const FVector& RayOrigin, const FVector& RayDirection)
+FInputRayHit UGizmoElementCone::LineTrace(const UGizmoViewContext* ViewContext, const FLineTraceTraversalState& LineTraceState, const FVector& RayOrigin, const FVector& RayDirection, FLineTraceOutput& OutLineTraceOutput)
 {
 	FLineTraceTraversalState CurrentLineTraceState(LineTraceState);
 	bool bHittableViewDependent = UpdateLineTraceState(ViewContext, Origin, CurrentLineTraceState);
@@ -108,15 +108,15 @@ FInputRayHit UGizmoElementCone::LineTrace(const UGizmoViewContext* ViewContext,
 		const double PixelHitThresholdAdjust = CurrentLineTraceState.PixelToWorldScale * PixelHitDistanceThreshold;
 		const double ConeSide = FMath::Sqrt(Height * Height + Radius * Radius);
 		const double CosAngle = Height / ConeSide;
-		const double WorldHeight = Height * CurrentLineTraceState.LocalToWorldTransform.GetScale3D().X + PixelHitThresholdAdjust * 2.0;
+		double WorldHeight = Height * CurrentLineTraceState.LocalToWorldTransform.GetScale3D().X + PixelHitThresholdAdjust * 2.0;
 		const FVector WorldDirection = CurrentLineTraceState.LocalToWorldTransform.TransformVectorNoScale(Direction);
-		const FVector WorldOrigin = CurrentLineTraceState.LocalToWorldTransform.TransformPosition(FVector::ZeroVector) - WorldDirection * PixelHitThresholdAdjust;
+		FVector WorldOrigin = CurrentLineTraceState.LocalToWorldTransform.TransformPosition(FVector::ZeroVector) - WorldDirection * PixelHitThresholdAdjust;
 
 		// due to numerical imprecision, the ray origin needs to be clamped in ortho views
 		// (cf. UEditorInteractiveToolsContext::GetRayFromMousePos)
 		FVector ClampedRayOrigin(RayOrigin);
 		const double DepthBias = UE::GizmoUtil::ClampRayOrigin(ViewContext, ClampedRayOrigin, RayDirection);
-		
+
 		GizmoMath::RayConeIntersection(
 			WorldOrigin,
 			WorldDirection,
@@ -129,9 +129,28 @@ FInputRayHit UGizmoElementCone::LineTrace(const UGizmoViewContext* ViewContext,
 		{
 			// add the depth bias if any
 			RayParam += DepthBias;
-			FInputRayHit RayHit(static_cast<float>(RayParam));
-			RayHit.SetHitObject(this);
-			RayHit.HitIdentifier = PartIdentifier;
+
+			FInputRayHit RayHit = MakeRayHit(RayParam, OutLineTraceOutput);
+			
+			// We successfully hit, now check if we intersect the actual surface (rather than just within the PixelHitDistanceThreshold) ...
+			const double MinimumPixelHitThresholdAdjust = CurrentLineTraceState.PixelToWorldScale * MinimumPixelHitDistanceThreshold;
+			WorldHeight = Height * CurrentLineTraceState.LocalToWorldTransform.GetScale3D().X + MinimumPixelHitThresholdAdjust * 2.0;
+			WorldOrigin = CurrentLineTraceState.LocalToWorldTransform.TransformPosition(FVector::ZeroVector) - WorldDirection * MinimumPixelHitThresholdAdjust;
+
+			GizmoMath::RayConeIntersection(
+				WorldOrigin,
+				WorldDirection,
+				CosAngle,
+				WorldHeight,
+				ClampedRayOrigin, RayDirection,
+				bIntersects, RayParam);
+
+			if (bIntersects)
+			{
+				RayHit.HitDepth = RayParam;
+				OutLineTraceOutput.bIsSurfaceHit = true;
+			}
+
 			return RayHit;
 		}
 	}
diff --git a/Engine/Source/Runtime/InteractiveToolsFramework/Private/BaseGizmos/GizmoElementCylinder.cpp b/Engine/Source/Runtime/InteractiveToolsFramework/Private/BaseGizmos/GizmoElementCylinder.cpp
--- a/Engine/Source/Runtime/InteractiveToolsFramework/Private/BaseGizmos/GizmoElementCylinder.cpp
+++ b/Engine/Source/Runtime/InteractiveToolsFramework/Private/BaseGizmos/GizmoElementCylinder.cpp
@@ -2,12 +2,16 @@
  
 #include "BaseGizmos/GizmoElementCylinder.h"
 
-#include "GizmoPrivateUtil.h"
-#include "BaseGizmos/GizmoRenderingUtil.h"
+#include "Algo/ForEach.h"
 #include "BaseGizmos/GizmoMath.h"
+#include "BaseGizmos/GizmoRenderingUtil.h"
+#include "DynamicMeshBuilder.h"
+#include "GizmoPrivateUtil.h"
 #include "InputState.h"
 #include "Materials/MaterialInterface.h"
 #include "PrimitiveDrawingUtils.h"
+#include "SceneManagement.h"
+#include "SphereTypes.h"
 
 #include UE_INLINE_GENERATED_CPP_BY_NAME(GizmoElementCylinder)
 
@@ -100,8 +104,7 @@ void UGizmoElementCylinder::Render(IToolsContextRenderAPI* RenderAPI, const FRen
 	}
 }
 
-
-FInputRayHit UGizmoElementCylinder::LineTrace(const UGizmoViewContext* ViewContext, const FLineTraceTraversalState& LineTraceState, const FVector& RayOrigin, const FVector& RayDirection)
+FInputRayHit UGizmoElementCylinder::LineTrace(const UGizmoViewContext* ViewContext, const FLineTraceTraversalState& LineTraceState, const FVector& RayOrigin, const FVector& RayDirection, FLineTraceOutput& OutLineTraceOutput)
 {
 	FLineTraceTraversalState CurrentLineTraceState(LineTraceState);
 	bool bHittableViewDependent = UpdateLineTraceState(ViewContext, Base, CurrentLineTraceState);
@@ -111,9 +114,9 @@ FInputRayHit UGizmoElementCylinder::LineTrace(const UGizmoViewContext* ViewConte
 		bool bIntersects = false;
 		double RayParam = 0.0;
 		
-		const double PixelHitThresholdAdjust = CurrentLineTraceState.PixelToWorldScale * PixelHitDistanceThreshold;
-		const double WorldHeight = Height * CurrentLineTraceState.LocalToWorldTransform.GetScale3D().X + PixelHitThresholdAdjust * 2.0;
-		const double WorldRadius = Radius * CurrentLineTraceState.LocalToWorldTransform.GetScale3D().X + PixelHitThresholdAdjust;
+		double PixelHitThresholdAdjust = CurrentLineTraceState.PixelToWorldScale * PixelHitDistanceThreshold;
+		double WorldHeight = Height * CurrentLineTraceState.LocalToWorldTransform.GetScale3D().X + PixelHitThresholdAdjust * 2.0;
+		double WorldRadius = Radius * CurrentLineTraceState.LocalToWorldTransform.GetScale3D().X + PixelHitThresholdAdjust;
 		const FVector WorldDirection = CurrentLineTraceState.LocalToWorldTransform.TransformVectorNoScale(Direction);
 		const FVector LocalCenter = Direction * Height * 0.5;
 		const FVector WorldCenter = CurrentLineTraceState.LocalToWorldTransform.TransformPosition(LocalCenter);
@@ -135,9 +138,40 @@ FInputRayHit UGizmoElementCylinder::LineTrace(const UGizmoViewContext* ViewConte
 		{
 			// add the depth bias if any
 			RayParam += DepthBias;
-			FInputRayHit RayHit(static_cast<float>(RayParam));
-			RayHit.SetHitObject(this);
-			RayHit.HitIdentifier = PartIdentifier;
+
+			FInputRayHit RayHit = MakeRayHit(RayParam, OutLineTraceOutput);
+
+			UE::Geometry::TRay<double> Ray(ClampedRayOrigin, RayDirection);
+			FVector HitPoint = Ray.PointAt(RayParam);
+
+			// We successfully hit, now check if we intersect the actual surface (rather than just within the PixelHitDistanceThreshold) ...
+			// ... by negating the threshold
+			WorldHeight -= PixelHitThresholdAdjust * 2.0;
+			WorldRadius -= PixelHitThresholdAdjust;
+
+			const double MinimumPixelHitThresholdAdjust = CurrentLineTraceState.PixelToWorldScale * MinimumPixelHitDistanceThreshold;
+			WorldHeight += MinimumPixelHitThresholdAdjust * 2.0;
+			WorldRadius += MinimumPixelHitThresholdAdjust;
+
+			// Clamp both height and radius to a minimum of PixelToWorldScale to ensure hittability
+			WorldHeight = FMath::Max(CurrentLineTraceState.PixelToWorldScale, WorldHeight);
+			WorldRadius = FMath::Max(CurrentLineTraceState.PixelToWorldScale, WorldRadius);
+
+			GizmoMath::RayCylinderIntersection(
+				WorldCenter,
+				WorldDirection,
+				WorldRadius,
+				WorldHeight,
+				ClampedRayOrigin, RayDirection,
+				bIntersects, RayParam);
+
+			// We hit the surface
+			if (bIntersects)
+			{
+				RayHit.HitDepth = RayParam + DepthBias;
+				OutLineTraceOutput.bIsSurfaceHit = true;
+			}
+
 			return RayHit;
 		}
 	}
diff --git a/Engine/Source/Runtime/InteractiveToolsFramework/Private/BaseGizmos/GizmoElementGroup.cpp b/Engine/Source/Runtime/InteractiveToolsFramework/Private/BaseGizmos/GizmoElementGroup.cpp
--- a/Engine/Source/Runtime/InteractiveToolsFramework/Private/BaseGizmos/GizmoElementGroup.cpp
+++ b/Engine/Source/Runtime/InteractiveToolsFramework/Private/BaseGizmos/GizmoElementGroup.cpp
@@ -56,7 +56,7 @@ void UGizmoElementGroupBase::DrawHUD(FCanvas* Canvas, IToolsContextRenderAPI* Re
 	}
 }
 
-FInputRayHit UGizmoElementGroupBase::LineTrace(const UGizmoViewContext* ViewContext, const FLineTraceTraversalState& LineTraceState, const FVector& RayOrigin, const FVector& RayDirection)
+FInputRayHit UGizmoElementGroupBase::LineTrace(const UGizmoViewContext* ViewContext, const FLineTraceTraversalState& LineTraceState, const FVector& RayOrigin, const FVector& RayDirection, FLineTraceOutput& OutLineTraceOutput)
 {
 	FInputRayHit Hit;
 
@@ -67,20 +67,86 @@ FInputRayHit UGizmoElementGroupBase::LineTrace(const UGizmoViewContext* ViewCont
 	{
 		ApplyUniformConstantScaleToTransform(CurrentLineTraceState.PixelToWorldScale, CurrentLineTraceState.LocalToWorldTransform);
 
-		ForEachSubElement([this, ViewContext, &CurrentLineTraceState, &RayOrigin, &RayDirection, &Hit](UGizmoElementBase* Element)
+		// This is similar to a sorting predicate, where the hit is only overridden if it's a better candidate
+		auto UpdateHit = [](
+			FInputRayHit& OutTargetOrExistingHit,
+			FLineTraceOutput& OutTargetOrExistingLineTraceOutput,
+			const FInputRayHit& InNewHit,
+			const FLineTraceOutput& InNewLineTraceOutput)
 		{
-			FInputRayHit NewHit = Element->LineTrace(ViewContext, CurrentLineTraceState, RayOrigin, RayDirection);
-			if (!Hit.bHit || NewHit.HitDepth < Hit.HitDepth)
+			// Early-out if the new hit is not valid
+			if (!InNewHit.bHit)
 			{
-				Hit = NewHit;
-				if (bHitOwner)
+				return;
+			}
+
+			const bool bIsOutputHitValid = OutTargetOrExistingHit.bHit;
+			bool bUseNewHit = !bIsOutputHitValid; // Use the new hit by default if the output hit is not valid
+
+			if (bIsOutputHitValid)
+			{
+				const bool bIsNegativePriority = InNewLineTraceOutput.HitPriority < 0;
+				const bool bIsHigherPriority = InNewLineTraceOutput.HitPriority > OutTargetOrExistingLineTraceOutput.HitPriority;
+				const bool bIsSamePriority = InNewLineTraceOutput.HitPriority == OutTargetOrExistingLineTraceOutput.HitPriority;
+				const bool bIsCloserHit = InNewHit.HitDepth < OutTargetOrExistingHit.HitDepth;
+
+				// Negative priorities are ALWAYS behind positive priorities, even if exact hits.
+				if (bIsNegativePriority && OutTargetOrExistingLineTraceOutput.HitPriority >= 0)
+				{
+					bUseNewHit = false;
+				}
+				// New is exact, current wasn't, use new
+				else if (InNewLineTraceOutput.bIsSurfaceHit && !OutTargetOrExistingLineTraceOutput.bIsSurfaceHit)
 				{
-					Hit.SetHitObject(this);
-					Hit.HitIdentifier = PartIdentifier;
+					bUseNewHit = true;
 				}
+				else if (!InNewLineTraceOutput.bIsSurfaceHit && OutTargetOrExistingLineTraceOutput.bIsSurfaceHit)
+				{
+					// Output hit is exact, new isn't, so keep output
+					bUseNewHit = false;
+				}
+				else if (bIsHigherPriority)
+				{
+					bUseNewHit = true;
+				}
+				else if (bIsSamePriority && bIsCloserHit)
+				{
+					bUseNewHit = true;
+				}
+				else
+				{
+					// If we're here, the new hit was lower priority, so we keep the output
+					bUseNewHit = false;
+				}
+			}
+
+			if (bUseNewHit)
+			{
+				OutTargetOrExistingHit = InNewHit;
 			}
+		};
+
+		ForEachSubElement([&UpdateHit, ViewContext, &CurrentLineTraceState, &RayOrigin, &RayDirection, &Hit, &OutLineTraceOutput](UGizmoElementBase* Element)
+		{
+			FLineTraceOutput NewLineTraceOutput;
+			const FInputRayHit NewHit = Element->LineTrace(ViewContext, CurrentLineTraceState, RayOrigin, RayDirection, NewLineTraceOutput);
+			if (!NewHit.bHit)
+			{
+				return;
+			}
+
+			UpdateHit(
+				Hit, OutLineTraceOutput,
+				NewHit, NewLineTraceOutput);
 		});
+
+		if (bHitOwner && Hit.bHit)
+		{
+			Hit.SetHitObject(this);
+			Hit.HitIdentifier = PartIdentifier;
+		}
 	}
+
 	return Hit;
 }
 
diff --git a/Engine/Source/Runtime/InteractiveToolsFramework/Private/BaseGizmos/GizmoElementHitTargets.cpp b/Engine/Source/Runtime/InteractiveToolsFramework/Private/BaseGizmos/GizmoElementHitTargets.cpp
--- a/Engine/Source/Runtime/InteractiveToolsFramework/Private/BaseGizmos/GizmoElementHitTargets.cpp
+++ b/Engine/Source/Runtime/InteractiveToolsFramework/Private/BaseGizmos/GizmoElementHitTargets.cpp
@@ -44,7 +44,9 @@ FInputRayHit UGizmoElementHitTarget::IsHit(const FInputDeviceRay& ClickPos) cons
 	{
 		UGizmoElementBase::FLineTraceTraversalState LineTraceState;
 		LineTraceState.Initialize(GizmoViewContext, GizmoTransformProxy->GetTransform());
-		return GizmoElement->LineTrace(GizmoViewContext, LineTraceState, ClickPos.WorldRay.Origin, ClickPos.WorldRay.Direction);
+
+		UGizmoElementBase::FLineTraceOutput LineTraceOutput;		
+		return GizmoElement->LineTrace(GizmoViewContext, LineTraceState, ClickPos.WorldRay.Origin, ClickPos.WorldRay.Direction, LineTraceOutput);
 	}
 	return FInputRayHit();
 }
@@ -103,7 +105,9 @@ FInputRayHit UGizmoElementHitMultiTarget::IsHit(const FInputDeviceRay& ClickPos)
 	{
 		UGizmoElementBase::FLineTraceTraversalState LineTraceState;
 		LineTraceState.Initialize(GizmoViewContext, GizmoTransformProxy->GetTransform());
-		return GizmoElement->LineTrace(GizmoViewContext, LineTraceState, ClickPos.WorldRay.Origin, ClickPos.WorldRay.Direction);
+
+		UGizmoElementBase::FLineTraceOutput LineTraceOutput;
+		return GizmoElement->LineTrace(GizmoViewContext, LineTraceState, ClickPos.WorldRay.Origin, ClickPos.WorldRay.Direction, LineTraceOutput);
 	}
 	return FInputRayHit();
 }
diff --git a/Engine/Source/Runtime/InteractiveToolsFramework/Private/BaseGizmos/GizmoElementLineStrip.cpp b/Engine/Source/Runtime/InteractiveToolsFramework/Private/BaseGizmos/GizmoElementLineStrip.cpp
--- a/Engine/Source/Runtime/InteractiveToolsFramework/Private/BaseGizmos/GizmoElementLineStrip.cpp
+++ b/Engine/Source/Runtime/InteractiveToolsFramework/Private/BaseGizmos/GizmoElementLineStrip.cpp
@@ -59,7 +59,7 @@ void UGizmoElementLineStrip::Render(IToolsContextRenderAPI* RenderAPI, const FRe
 	}
 }
 
-FInputRayHit UGizmoElementLineStrip::LineTrace(const UGizmoViewContext* ViewContext, const FLineTraceTraversalState& LineTraceState, const FVector& RayOrigin, const FVector& RayDirection)
+FInputRayHit UGizmoElementLineStrip::LineTrace(const UGizmoViewContext* ViewContext, const FLineTraceTraversalState& LineTraceState, const FVector& RayOrigin, const FVector& RayDirection, FLineTraceOutput& OutLineTraceOutput)
 {
 	FLineTraceTraversalState CurrentLineTraceState(LineTraceState);
 	const bool bHittableViewDependent = UpdateLineTraceState(ViewContext, Base, CurrentLineTraceState);
@@ -94,10 +94,7 @@ FInputRayHit UGizmoElementLineStrip::LineTrace(const UGizmoViewContext* ViewCont
 				const FVector& B = ProjectedVertices[Index+1];
 				if(Intersect(A, B, PixelHitThresholdAdjust, HitDepth))
 				{
-					FInputRayHit RayHit(static_cast<float>(HitDepth));
-					RayHit.SetHitObject(this);
-					RayHit.HitIdentifier = PartIdentifier;
-					return RayHit;
+					return MakeRayHit(HitDepth, OutLineTraceOutput);
 				}
 			}
 		}
@@ -109,14 +106,12 @@ FInputRayHit UGizmoElementLineStrip::LineTrace(const UGizmoViewContext* ViewCont
 				const FVector& B = ProjectedVertices[Index+1];
 				if(Intersect(A, B, PixelHitThresholdAdjust, HitDepth))
 				{
-					FInputRayHit RayHit(static_cast<float>(HitDepth));
-					RayHit.SetHitObject(this);
-					RayHit.HitIdentifier = PartIdentifier;
-					return RayHit;
+					return MakeRayHit(HitDepth, OutLineTraceOutput);
 				}
 			}
 		}
 	}
+
 	return FInputRayHit();
 }
 
diff --git a/Engine/Source/Runtime/InteractiveToolsFramework/Private/BaseGizmos/GizmoElementRectangle.cpp b/Engine/Source/Runtime/InteractiveToolsFramework/Private/BaseGizmos/GizmoElementRectangle.cpp
--- a/Engine/Source/Runtime/InteractiveToolsFramework/Private/BaseGizmos/GizmoElementRectangle.cpp
+++ b/Engine/Source/Runtime/InteractiveToolsFramework/Private/BaseGizmos/GizmoElementRectangle.cpp
@@ -58,7 +58,7 @@ void UGizmoElementRectangle::Render(IToolsContextRenderAPI* RenderAPI, const FRe
 	}
 }
 
-FInputRayHit UGizmoElementRectangle::LineTrace(const UGizmoViewContext* ViewContext, const FLineTraceTraversalState& LineTraceState, const FVector& RayOrigin, const FVector& RayDirection)
+FInputRayHit UGizmoElementRectangle::LineTrace(const UGizmoViewContext* ViewContext, const FLineTraceTraversalState& LineTraceState, const FVector& RayOrigin, const FVector& RayDirection, FLineTraceOutput& OutLineTraceOutput)
 {
 	FLineTraceTraversalState CurrentLineTraceState(LineTraceState);
 	bool bHittableViewDependent = UpdateLineTraceState(ViewContext, Center, CurrentLineTraceState);
@@ -72,8 +72,8 @@ FInputRayHit UGizmoElementRectangle::LineTrace(const UGizmoViewContext* ViewCont
 		const FVector WorldCenter = LocalToWorldTransform.TransformPosition(FVector::ZeroVector);
 		const double Scale = LocalToWorldTransform.GetScale3D().X;
 		const double PixelHitThresholdAdjust = CurrentLineTraceState.PixelToWorldScale * PixelHitDistanceThreshold;
-		const double WorldHeight = Scale * Height + PixelHitThresholdAdjust * 2.0;
-		const double WorldWidth = Scale * Width + PixelHitThresholdAdjust * 2.0;
+		double WorldHeight = Scale * Height + PixelHitThresholdAdjust * 2.0;
+		double WorldWidth = Scale * Width + PixelHitThresholdAdjust * 2.0;
 		const FVector Base = WorldCenter - WorldUpAxis * WorldHeight * 0.5 - WorldSideAxis * WorldWidth * 0.5;
 
 		// if ray is parallel to rectangle, no hit
@@ -86,7 +86,7 @@ FInputRayHit UGizmoElementRectangle::LineTrace(const UGizmoViewContext* ViewCont
 		{
 			FPlane Plane(Base, WorldNormal);
 			double HitDepth = FMath::RayPlaneIntersectionParam(RayOrigin, RayDirection, Plane);
-			if (HitDepth < 0)
+			if (HitDepth < UE_DOUBLE_SMALL_NUMBER)
 			{
 				return FInputRayHit();
 			}
@@ -97,11 +97,16 @@ FInputRayHit UGizmoElementRectangle::LineTrace(const UGizmoViewContext* ViewCont
 			double HdS = FVector::DotProduct(HitOffset, WorldSideAxis);
 
 			// clip to rectangle dimensions
-			if (HdU >= 0.0 && HdU <= WorldHeight && HdS >= 0.0 && HdS <= WorldWidth)
+			if (HdU >= UE_DOUBLE_SMALL_NUMBER && HdU <= WorldHeight && HdS >= UE_DOUBLE_SMALL_NUMBER && HdS <= WorldWidth)
 			{
-				FInputRayHit RayHit(static_cast<float>(HitDepth));
-				RayHit.SetHitObject(this);
-				RayHit.HitIdentifier = PartIdentifier;
+				FInputRayHit RayHit = MakeRayHit(HitDepth, OutLineTraceOutput);
+
+				// We successfully hit, now check if we intersect the actual surface (rather than just within the PixelHitDistanceThreshold)
+				const double MinimumPixelHitThresholdAdjust = CurrentLineTraceState.PixelToWorldScale * MinimumPixelHitDistanceThreshold;
+				WorldHeight = Scale * Height + MinimumPixelHitThresholdAdjust * 2.0;
+				WorldWidth = Scale * Width + MinimumPixelHitThresholdAdjust * 2.0;
+				OutLineTraceOutput.bIsSurfaceHit = HdU <= WorldHeight && HdS <= WorldWidth;
+
 				return RayHit;
 			}
 		}
@@ -109,7 +114,7 @@ FInputRayHit UGizmoElementRectangle::LineTrace(const UGizmoViewContext* ViewCont
 		{
 			FPlane Plane(Base, WorldNormal);
 			double HitDepth = FMath::RayPlaneIntersectionParam(RayOrigin, RayDirection, Plane);
-			if (HitDepth < 0)
+			if (HitDepth < UE_DOUBLE_SMALL_NUMBER)
 			{
 				return FInputRayHit();
 			}
@@ -119,20 +124,26 @@ FInputRayHit UGizmoElementRectangle::LineTrace(const UGizmoViewContext* ViewCont
 			double HdU = FVector::DotProduct(HitOffset, WorldUpAxis);
 			double HdS = FVector::DotProduct(HitOffset, WorldSideAxis);
 
-			const double HitBuffer = PixelHitThresholdAdjust + LineThickness;
+			double HitBuffer = PixelHitThresholdAdjust + LineThickness;
 
 			// determine if the hit is within pixel tolerance of the edges of rectangle
-			if (HdU >= 0.0 && HdU <= WorldHeight && HdS >= 0.0 && HdS <= WorldWidth &&
+			if (HdU >= UE_DOUBLE_SMALL_NUMBER && HdU <= WorldHeight && HdS >= UE_DOUBLE_SMALL_NUMBER && HdS <= WorldWidth &&
 				(HdS <= HitBuffer || HdS >= WorldWidth - HitBuffer ||
 				 HdU <= HitBuffer || HdU >= WorldHeight - HitBuffer))
 			{
-				FInputRayHit RayHit(static_cast<float>(HitDepth));
-				RayHit.SetHitObject(this);
-				RayHit.HitIdentifier = PartIdentifier;
+				FInputRayHit RayHit = MakeRayHit(HitDepth, OutLineTraceOutput);
+
+				// We successfully hit, now check if we intersect the actual surface (rather than just within the PixelHitDistanceThreshold)
+				const double MinimumPixelHitThresholdAdjust = CurrentLineTraceState.PixelToWorldScale * MinimumPixelHitDistanceThreshold;
+				HitBuffer = MinimumPixelHitThresholdAdjust + LineThickness;
+				OutLineTraceOutput.bIsSurfaceHit = (HdS <= HitBuffer || HdS >= WorldWidth - HitBuffer ||
+					HdU <= HitBuffer || HdU >= WorldHeight - HitBuffer);
+
 				return RayHit;
 			}
 		}
 	}
+
 	return FInputRayHit();
 }
 
diff --git a/Engine/Source/Runtime/InteractiveToolsFramework/Private/BaseGizmos/GizmoElementSphere.cpp b/Engine/Source/Runtime/InteractiveToolsFramework/Private/BaseGizmos/GizmoElementSphere.cpp
--- a/Engine/Source/Runtime/InteractiveToolsFramework/Private/BaseGizmos/GizmoElementSphere.cpp
+++ b/Engine/Source/Runtime/InteractiveToolsFramework/Private/BaseGizmos/GizmoElementSphere.cpp
@@ -31,7 +31,7 @@ void UGizmoElementSphere::Render(IToolsContextRenderAPI* RenderAPI, const FRende
 	}
 }
 
-FInputRayHit UGizmoElementSphere::LineTrace(const UGizmoViewContext* ViewContext, const FLineTraceTraversalState& LineTraceState, const FVector& RayOrigin, const FVector& RayDirection)
+FInputRayHit UGizmoElementSphere::LineTrace(const UGizmoViewContext* ViewContext, const FLineTraceTraversalState& LineTraceState, const FVector& RayOrigin, const FVector& RayDirection, FLineTraceOutput& OutLineTraceOutput)
 {
 	FLineTraceTraversalState CurrentLineTraceState(LineTraceState);
 	bool bHittableViewDependent = UpdateLineTraceState(ViewContext, Center, CurrentLineTraceState);
@@ -48,10 +48,7 @@ FInputRayHit UGizmoElementSphere::LineTrace(const UGizmoViewContext* ViewContext
 
 		if(HitDepth < WorldExtent)
 		{
-			FInputRayHit RayHit(static_cast<float>(HitDepth));
-			RayHit.SetHitObject(this);
-			RayHit.HitIdentifier = PartIdentifier;
-			return RayHit;
+			return MakeRayHit(HitDepth, OutLineTraceOutput);
 		}
 	}
 
diff --git a/Engine/Source/Runtime/InteractiveToolsFramework/Private/BaseGizmos/GizmoElementTorus.cpp b/Engine/Source/Runtime/InteractiveToolsFramework/Private/BaseGizmos/GizmoElementTorus.cpp
--- a/Engine/Source/Runtime/InteractiveToolsFramework/Private/BaseGizmos/GizmoElementTorus.cpp
+++ b/Engine/Source/Runtime/InteractiveToolsFramework/Private/BaseGizmos/GizmoElementTorus.cpp
@@ -47,7 +47,6 @@ void UGizmoElementTorus::Render(IToolsContextRenderAPI* RenderAPI, const FRender
 	}
 }
 
-
 //
 // This method approximates ray-torus intersection by intersecting the ray with the plane in which the torus 
 // lies, then determining a hit point closest to the linear circle defined by torus center and torus outer radius.
@@ -55,7 +54,7 @@ void UGizmoElementTorus::Render(IToolsContextRenderAPI* RenderAPI, const FRender
 // If the torus lies at a glancing angle, the ray-torus intersection is performed against cylinders approximating
 // the shape of the torus.
 //
-FInputRayHit UGizmoElementTorus::LineTrace(const UGizmoViewContext* ViewContext, const FLineTraceTraversalState& LineTraceState, const FVector& RayOrigin, const FVector& RayDirection)
+FInputRayHit UGizmoElementTorus::LineTrace(const UGizmoViewContext* ViewContext, const FLineTraceTraversalState& LineTraceState, const FVector& RayOrigin, const FVector& RayDirection, FLineTraceOutput& OutLineTraceOutput)
 {
 	FLineTraceTraversalState CurrentLineTraceState(LineTraceState);
 	bool bHittableViewDependent = UpdateLineTraceState(ViewContext, Center, CurrentLineTraceState);
@@ -70,6 +69,7 @@ FInputRayHit UGizmoElementTorus::LineTrace(const UGizmoViewContext* ViewContext,
 		const FVector WorldAxis1 = CurrentLineTraceState.LocalToWorldTransform.TransformVectorNoScale(Axis1);
 		const FVector WorldNormal = WorldAxis0 ^ WorldAxis1;
 		double HitDepth = -1.0;
+		bool bIsExactHit = false;
 
 		bool bPartial = IsPartial(ViewContext, WorldCenter, WorldNormal);
 
@@ -109,6 +109,9 @@ FInputRayHit UGizmoElementTorus::LineTrace(const UGizmoViewContext* ViewContext,
 				return FInputRayHit();
 			}
 
+			const double MinimumPixelHitThresholdAdjust = CurrentLineTraceState.PixelToWorldScale * MinimumPixelHitDistanceThreshold;
+			const double ExactCylinderRadius = FMath::Max(CurrentLineTraceState.PixelToWorldScale, WorldInnerRadius + MinimumPixelHitThresholdAdjust);
+
 			// due to numerical imprecision, the ray origin needs to be clamped in ortho views
 			// (cf. UEditorInteractiveToolsContext::GetRayFromMousePos)
 			FVector ClampedRayOrigin(RayOrigin);
@@ -146,9 +149,20 @@ FInputRayHit UGizmoElementTorus::LineTrace(const UGizmoViewContext* ViewContext,
 				bDidIntersect |= bIntersects; 
 				
 				// Update with closest hit depth
-				if (bIntersects && (HitDepth < 0 || HitDepth > RayParam))
+				if (bIntersects && (HitDepth < UE_DOUBLE_SMALL_NUMBER || HitDepth > RayParam))
 				{
 					HitDepth = RayParam;
+
+					// We successfully hit, now check if we intersect the actual surface (rather than just within the PixelHitDistanceThreshold)
+					GizmoMath::RayCylinderIntersection(
+						CylinderCenter,
+						CylinderDirection,
+						ExactCylinderRadius,
+						CylinderHeight,
+						ClampedRayOrigin, RayDirection,
+						bIntersects, RayParam);
+
+					bIsExactHit = bIntersects;
 				}
 			}
 
@@ -163,7 +177,7 @@ FInputRayHit UGizmoElementTorus::LineTrace(const UGizmoViewContext* ViewContext,
 			// Intersect ray with plane in which torus lies.
 			FPlane Plane(WorldCenter, WorldNormal);
 			HitDepth = FMath::RayPlaneIntersectionParam(RayOrigin, RayDirection, Plane);
-			if (HitDepth < 0.0)
+			if (HitDepth < UE_DOUBLE_SMALL_NUMBER)
 			{
 				return FInputRayHit();
 			}
@@ -186,6 +200,8 @@ FInputRayHit UGizmoElementTorus::LineTrace(const UGizmoViewContext* ViewContext,
 				return FInputRayHit();
 			}
 
+			bIsExactHit = (Distance <= FMath::Max(CurrentLineTraceState.PixelToWorldScale * 2.0f, WorldInnerRadius));
+
 			// Handle partial torus
 			if (bPartial)
 			{
@@ -193,7 +209,7 @@ FInputRayHit UGizmoElementTorus::LineTrace(const UGizmoViewContext* ViewContext,
 				FVector HitVec = (NearestCirclePos - WorldCenter).GetSafeNormal();
 				double HitAngle = UE::Geometry::VectorUtil::PlaneAngleSignedR(WorldBeginAxis, HitVec, WorldNormal); 
 
-				if (HitAngle < 0)
+				if (HitAngle < UE_DOUBLE_SMALL_NUMBER)
 				{
 					HitAngle = UE_DOUBLE_TWO_PI + HitAngle;
 				}
@@ -208,11 +224,11 @@ FInputRayHit UGizmoElementTorus::LineTrace(const UGizmoViewContext* ViewContext,
 			HitDepth = (NearestRayPos - RayOrigin).Length();
 		}
 
-		if (HitDepth >= 0.0)
+		if (HitDepth >= UE_DOUBLE_SMALL_NUMBER)
 		{
-			FInputRayHit RayHit(static_cast<float>(HitDepth));
-			RayHit.SetHitObject(this);
-			RayHit.HitIdentifier = PartIdentifier;
+			FInputRayHit RayHit = MakeRayHit(HitDepth, OutLineTraceOutput);
+			OutLineTraceOutput.bIsSurfaceHit = bIsExactHit;
+
 			return RayHit;
 		}
 	}
diff --git a/Engine/Source/Runtime/InteractiveToolsFramework/Private/BaseGizmos/GizmoElementTriangleList.cpp b/Engine/Source/Runtime/InteractiveToolsFramework/Private/BaseGizmos/GizmoElementTriangleList.cpp
--- a/Engine/Source/Runtime/InteractiveToolsFramework/Private/BaseGizmos/GizmoElementTriangleList.cpp
+++ b/Engine/Source/Runtime/InteractiveToolsFramework/Private/BaseGizmos/GizmoElementTriangleList.cpp
@@ -35,7 +35,7 @@ void UGizmoElementTriangleList::Render(IToolsContextRenderAPI* RenderAPI, const
 	}
 }
 
-FInputRayHit UGizmoElementTriangleList::LineTrace(const UGizmoViewContext* ViewContext, const FLineTraceTraversalState& LineTraceState, const FVector& RayOrigin, const FVector& RayDirection)
+FInputRayHit UGizmoElementTriangleList::LineTrace(const UGizmoViewContext* ViewContext, const FLineTraceTraversalState& LineTraceState, const FVector& RayOrigin, const FVector& RayDirection, FLineTraceOutput& OutLineTraceOutput)
 {
 	FLineTraceTraversalState CurrentLineTraceState(LineTraceState);
 	const bool bHittableViewDependent = UpdateLineTraceState(ViewContext, Base, CurrentLineTraceState);
@@ -55,10 +55,7 @@ FInputRayHit UGizmoElementTriangleList::LineTrace(const UGizmoViewContext* ViewC
 			const FVector& C = ProjectedVertices[Index++];
 			if (Intersect(Ray, A, B, C, PixelHitThresholdAdjust, HitDepth))
 			{
-				FInputRayHit RayHit(static_cast<float>(HitDepth));
-				RayHit.SetHitObject(this);
-				RayHit.HitIdentifier = PartIdentifier;
-				return RayHit;
+				return MakeRayHit(HitDepth, OutLineTraceOutput);
 			}
 		}
 	}
diff --git a/Engine/Source/Runtime/InteractiveToolsFramework/Public/BaseGizmos/GizmoElementArc.h b/Engine/Source/Runtime/InteractiveToolsFramework/Public/BaseGizmos/GizmoElementArc.h
--- a/Engine/Source/Runtime/InteractiveToolsFramework/Public/BaseGizmos/GizmoElementArc.h
+++ b/Engine/Source/Runtime/InteractiveToolsFramework/Public/BaseGizmos/GizmoElementArc.h
@@ -26,7 +26,7 @@ public:
 
 	//~ Begin UGizmoElementBase Interface.
 	INTERACTIVETOOLSFRAMEWORK_API virtual void Render(IToolsContextRenderAPI* RenderAPI, const FRenderTraversalState& RenderState) override;
-	INTERACTIVETOOLSFRAMEWORK_API virtual FInputRayHit LineTrace(const UGizmoViewContext* ViewContext, const FLineTraceTraversalState& LineTraceState, const FVector& RayOrigin, const FVector& RayDirection) override;
+	INTERACTIVETOOLSFRAMEWORK_API virtual FInputRayHit LineTrace(const UGizmoViewContext* ViewContext, const FLineTraceTraversalState& LineTraceState, const FVector& RayOrigin, const FVector& RayDirection, FLineTraceOutput& OutLineTraceOutput) override;
 	//~ End UGizmoElementBase Interface.
 
 	// Inner circle radius.
diff --git a/Engine/Source/Runtime/InteractiveToolsFramework/Public/BaseGizmos/GizmoElementArrow.h b/Engine/Source/Runtime/InteractiveToolsFramework/Public/BaseGizmos/GizmoElementArrow.h
--- a/Engine/Source/Runtime/InteractiveToolsFramework/Public/BaseGizmos/GizmoElementArrow.h
+++ b/Engine/Source/Runtime/InteractiveToolsFramework/Public/BaseGizmos/GizmoElementArrow.h
@@ -30,10 +30,10 @@ public:
 
 	//~ Begin UGizmoElementBase Interface.
 	INTERACTIVETOOLSFRAMEWORK_API virtual void Render(IToolsContextRenderAPI* RenderAPI, const FRenderTraversalState& RenderState) override;
-	INTERACTIVETOOLSFRAMEWORK_API virtual FInputRayHit LineTrace(const UGizmoViewContext* ViewContext, const FLineTraceTraversalState& LineTraceState, const FVector& RayOrigin, const FVector& RayDirection) override;
+	INTERACTIVETOOLSFRAMEWORK_API virtual FInputRayHit LineTrace(const UGizmoViewContext* ViewContext, const FLineTraceTraversalState& LineTraceState, const FVector& RayOrigin, const FVector& RayDirection, FLineTraceOutput& OutLineTraceOutput) override;
 	//~ End UGizmoElementBase Interface.
 
-	// Location of base of arrow cylinder
+	// Location of the base of arrow cylinder
 	INTERACTIVETOOLSFRAMEWORK_API virtual void SetBase(const FVector& InBase);
 	INTERACTIVETOOLSFRAMEWORK_API virtual FVector GetBase() const;
 
@@ -84,6 +84,9 @@ public:
 	// Pixel hit distance threshold, element will be scaled enough to add this threshold when line-tracing.
 	INTERACTIVETOOLSFRAMEWORK_API virtual void SetPixelHitDistanceThreshold(float InPixelHitDistanceThreshold) override;
 
+	// Minimum Pixel hit distance threshold, element will be scaled enough to add this threshold when line-tracing.
+	INTERACTIVETOOLSFRAMEWORK_API virtual void SetMinimumPixelHitDistanceThreshold(float InMinimumPixelHitDistanceThreshold) override;
+	
 	// Hit Mask Gizmo element used for adaptive pixel hit threshold
 	INTERACTIVETOOLSFRAMEWORK_API void SetHitMask(const TWeakObjectPtr<UGizmoElementBase>& InHitMask);
 
diff --git a/Engine/Source/Runtime/InteractiveToolsFramework/Public/BaseGizmos/GizmoElementArrowHead.h b/Engine/Source/Runtime/InteractiveToolsFramework/Public/BaseGizmos/GizmoElementArrowHead.h
--- a/Engine/Source/Runtime/InteractiveToolsFramework/Public/BaseGizmos/GizmoElementArrowHead.h
+++ b/Engine/Source/Runtime/InteractiveToolsFramework/Public/BaseGizmos/GizmoElementArrowHead.h
@@ -34,7 +34,7 @@ public:
 
 	//~ Begin UGizmoElementBase Interface.
 	INTERACTIVETOOLSFRAMEWORK_API virtual void Render(IToolsContextRenderAPI* RenderAPI, const FRenderTraversalState& RenderState) override;
-	INTERACTIVETOOLSFRAMEWORK_API virtual FInputRayHit LineTrace(const UGizmoViewContext* ViewContext, const FLineTraceTraversalState& LineTraceState, const FVector& RayOrigin, const FVector& RayDirection) override;
+	INTERACTIVETOOLSFRAMEWORK_API virtual FInputRayHit LineTrace(const UGizmoViewContext* ViewContext, const FLineTraceTraversalState& LineTraceState, const FVector& RayOrigin, const FVector& RayDirection, FLineTraceOutput& OutLineTraceOutput) override;
 	//~ End UGizmoElementBase Interface.
 
 	// Location of center of the arrow head
diff --git a/Engine/Source/Runtime/InteractiveToolsFramework/Public/BaseGizmos/GizmoElementBase.h b/Engine/Source/Runtime/InteractiveToolsFramework/Public/BaseGizmos/GizmoElementBase.h
--- a/Engine/Source/Runtime/InteractiveToolsFramework/Public/BaseGizmos/GizmoElementBase.h
+++ b/Engine/Source/Runtime/InteractiveToolsFramework/Public/BaseGizmos/GizmoElementBase.h
@@ -54,7 +54,7 @@ public:
 		// Note: non-uniform scale is not supported and the X scale element will be used for uniform scaling.
 		FTransform LocalToWorldTransform;
 
-		// Pixel to world scale. To account for DPI scale, @see IToolsContextRenderAPI::ViewCameraState.
+		// Pixel to world scale. Also accounts for DPI scale, @see IToolsContextRenderAPI::ViewCameraState.
 		double PixelToWorldScale = 1.0;
 
 		// Interact state, if not equal to none, overrides the element's interact state 
@@ -109,11 +109,21 @@ public:
 		{
 			check(InGizmoViewContext);
 			LocalToWorldTransform = InTransform;
-			PixelToWorldScale = UE::GizmoRenderingUtil::CalculateLocalPixelToWorldScale(InGizmoViewContext, InTransform.GetLocation());
+			PixelToWorldScale = UE::GizmoRenderingUtil::CalculateLocalPixelToWorldScale(InGizmoViewContext, InTransform.GetLocation(), true, InGizmoViewContext->GetDPIScale());
 			bIsPerspectiveProjection = InGizmoViewContext->IsPerspectiveProjection();
 		}
 	};
 
+	/** Contains information about a line trace, in addition to FInputRayHit. */
+	struct FLineTraceOutput
+	{
+		/** Client-defined hit priority. Higher priorities can optionally be used to prefer one hit over another */
+		int32 HitPriority = 0;
+
+		/** True if the ray hit the actual surface, rather than within a proximity (specified by PixelHitDistanceThreshold in GizmoElements). */
+		bool bIsSurfaceHit = false;
+	};
+
 public:
 
 	// Render enabled visible element.
@@ -121,9 +131,13 @@ public:
 
 	// Do screen space drawing.
 	INTERACTIVETOOLSFRAMEWORK_API virtual void DrawHUD(FCanvas* Canvas, IToolsContextRenderAPI* RenderAPI, const FRenderTraversalState& RenderState);
+	
+	// Line trace enabled hittable element.
+	virtual FInputRayHit LineTrace(const UGizmoViewContext* ViewContext, const FLineTraceTraversalState& LineTraceState, const FVector& RayOrigin, const FVector& RayDirection, FLineTraceOutput& OutLineTraceOutput) PURE_VIRTUAL(UGizmoElementBase::LineTrace, return FInputRayHit(););
 
 	// Line trace enabled hittable element.
-	virtual FInputRayHit LineTrace(const UGizmoViewContext* ViewContext, const FLineTraceTraversalState& LineTraceState, const FVector& RayOrigin, const FVector& RayDirection) PURE_VIRTUAL(UGizmoElementBase::LineTrace, return FInputRayHit(););
+	UE_DEPRECATED(5.7, "Use the version of LineTrace that outputs FLineTraceOutput instead.")
+	INTERACTIVETOOLSFRAMEWORK_API virtual FInputRayHit LineTrace(const UGizmoViewContext* ViewContext, const FLineTraceTraversalState& LineTraceState, const FVector& RayOrigin, const FVector& RayDirection);
 
 	// Set/get the visible bit in element state.
 	INTERACTIVETOOLSFRAMEWORK_API virtual void SetVisibleState(bool bVisible);
@@ -133,6 +147,10 @@ public:
 	INTERACTIVETOOLSFRAMEWORK_API virtual void SetHittableState(bool bHittable);
 	INTERACTIVETOOLSFRAMEWORK_API virtual bool GetHittableState() const;
 
+	// Set/get the hit priority. When hits overlap, higher priority elements will be hit first. If two elements have the same priority, the one that was closest to the hit location will be returned.
+	INTERACTIVETOOLSFRAMEWORK_API virtual void SetHitPriority(const int32 InHitPriority);
+	INTERACTIVETOOLSFRAMEWORK_API virtual int32 GetHitPriority() const;
+
 	// Element enabled flag. Render and LineTrace only occur when bEnabled is true.
 	// This flag is useful for turning on and off an element globally while retaining its specific settings.
 	INTERACTIVETOOLSFRAMEWORK_API virtual void SetEnabled(bool bInEnabled);
@@ -276,6 +294,10 @@ public:
 	INTERACTIVETOOLSFRAMEWORK_API virtual void SetPixelHitDistanceThreshold(float InPixelHitDistanceThreshold);
 	INTERACTIVETOOLSFRAMEWORK_API virtual float GetPixelHitDistanceThreshold() const;
 
+	// Minimum pixel hit distance threshold, element will be scaled enough to add this threshold when line-tracing the "exact" element (useful for pixel-thin elements where an "exact" hit would be otherwise challenging to hit).
+	INTERACTIVETOOLSFRAMEWORK_API virtual void SetMinimumPixelHitDistanceThreshold(const float InMinimumPixelHitDistanceThreshold);
+	INTERACTIVETOOLSFRAMEWORK_API virtual float GetMinimumPixelHitDistanceThreshold() const;
+
 	//
 	// Methods for managing render state attributes: Material, HoverMaterial, InteractMaterial, SelectMaterial, SubdueMaterial, VertexColor 
 	// 
@@ -546,6 +568,14 @@ protected:
 	UPROPERTY()
 	float PixelHitDistanceThreshold = 7.0;
 
+	// Minimum pixel hit distance threshold, element will be scaled enough to add this threshold when line-tracing the "exact" element (useful for pixel-thin elements where an "exact" hit would be otherwise challenging to hit).
+	UPROPERTY()
+	float MinimumPixelHitDistanceThreshold = 0.0;
+
+	// When hits overlap, higher priority elements will be hit first. If two elements have the same priority, the one that was closest to the hit location will be returned.
+	UPROPERTY()
+	int32 HitPriority = 0;
+
 private:
 
 	// Additional view-space offset to apply to the element, affecting both visuals and the hit geometry.
@@ -650,6 +680,9 @@ protected:
 	// @param OutAlignRot - alignment rotation applied to output state, if applicable.
 	INTERACTIVETOOLSFRAMEWORK_API virtual bool UpdateLineTraceState(const UGizmoViewContext* ViewContext, const FVector& InLocalCenter, FLineTraceTraversalState& InOutRenderState, bool& bOutHasAlignRot, FQuat& OutAlignRot);
 
+	// Initialize ray hit for this element, injecting the Part Identifier, Hit Priority, etc.
+	[[nodiscard]] INTERACTIVETOOLSFRAMEWORK_API virtual FInputRayHit MakeRayHit(const double InHitDepth, FLineTraceOutput& OutLineTraceOutput);
+
 protected:
 
 	// Helper method to verify scale is uniform. If it is non-uniform, a one-time warning log is issued.
diff --git a/Engine/Source/Runtime/InteractiveToolsFramework/Public/BaseGizmos/GizmoElementBox.h b/Engine/Source/Runtime/InteractiveToolsFramework/Public/BaseGizmos/GizmoElementBox.h
--- a/Engine/Source/Runtime/InteractiveToolsFramework/Public/BaseGizmos/GizmoElementBox.h
+++ b/Engine/Source/Runtime/InteractiveToolsFramework/Public/BaseGizmos/GizmoElementBox.h
@@ -22,7 +22,7 @@ public:
 
 	//~ Begin UGizmoElementBase Interface.
 	INTERACTIVETOOLSFRAMEWORK_API virtual void Render(IToolsContextRenderAPI* RenderAPI, const FRenderTraversalState& RenderState) override;
-	INTERACTIVETOOLSFRAMEWORK_API virtual FInputRayHit LineTrace(const UGizmoViewContext* ViewContext, const FLineTraceTraversalState& LineTraceState, const FVector& RayOrigin, const FVector& RayDirection) override;
+	INTERACTIVETOOLSFRAMEWORK_API virtual FInputRayHit LineTrace(const UGizmoViewContext* ViewContext, const FLineTraceTraversalState& LineTraceState, const FVector& RayOrigin, const FVector& RayDirection, FLineTraceOutput& OutLineTraceOutput) override;
 	//~ End UGizmoElementBase Interface.
 
 	// Location of center of box's base circle.
diff --git a/Engine/Source/Runtime/InteractiveToolsFramework/Public/BaseGizmos/GizmoElementCircle.h b/Engine/Source/Runtime/InteractiveToolsFramework/Public/BaseGizmos/GizmoElementCircle.h
--- a/Engine/Source/Runtime/InteractiveToolsFramework/Public/BaseGizmos/GizmoElementCircle.h
+++ b/Engine/Source/Runtime/InteractiveToolsFramework/Public/BaseGizmos/GizmoElementCircle.h
@@ -22,7 +22,7 @@ class UGizmoElementCircle : public UGizmoElementCircleBase
 public:
 	//~ Begin UGizmoElementBase Interface.
 	INTERACTIVETOOLSFRAMEWORK_API virtual void Render(IToolsContextRenderAPI* RenderAPI, const FRenderTraversalState& RenderState) override;
-	INTERACTIVETOOLSFRAMEWORK_API virtual FInputRayHit LineTrace(const UGizmoViewContext* ViewContext, const FLineTraceTraversalState& LineTraceState, const FVector& RayOrigin, const FVector& RayDirection) override;
+	INTERACTIVETOOLSFRAMEWORK_API virtual FInputRayHit LineTrace(const UGizmoViewContext* ViewContext, const FLineTraceTraversalState& LineTraceState, const FVector& RayOrigin, const FVector& RayDirection, FLineTraceOutput& OutLineTraceOutput) override;
 	//~ End UGizmoElementBase Interface.
 
 	// Draw mesh
diff --git a/Engine/Source/Runtime/InteractiveToolsFramework/Public/BaseGizmos/GizmoElementCone.h b/Engine/Source/Runtime/InteractiveToolsFramework/Public/BaseGizmos/GizmoElementCone.h
--- a/Engine/Source/Runtime/InteractiveToolsFramework/Public/BaseGizmos/GizmoElementCone.h
+++ b/Engine/Source/Runtime/InteractiveToolsFramework/Public/BaseGizmos/GizmoElementCone.h
@@ -22,7 +22,7 @@ public:
 
 	//~ Begin UGizmoElementBase Interface.
 	INTERACTIVETOOLSFRAMEWORK_API virtual void Render(IToolsContextRenderAPI* RenderAPI, const FRenderTraversalState& RenderState) override;
-	INTERACTIVETOOLSFRAMEWORK_API virtual FInputRayHit LineTrace(const UGizmoViewContext* ViewContext, const FLineTraceTraversalState& LineTraceState, const FVector& RayOrigin, const FVector& RayDirection) override;
+	INTERACTIVETOOLSFRAMEWORK_API virtual FInputRayHit LineTrace(const UGizmoViewContext* ViewContext, const FLineTraceTraversalState& LineTraceState, const FVector& RayOrigin, const FVector& RayDirection, FLineTraceOutput& OutLineTraceOutput) override;
 	//~ End UGizmoElementBase Interface.
 
 	// Cone tip location.
diff --git a/Engine/Source/Runtime/InteractiveToolsFramework/Public/BaseGizmos/GizmoElementCylinder.h b/Engine/Source/Runtime/InteractiveToolsFramework/Public/BaseGizmos/GizmoElementCylinder.h
--- a/Engine/Source/Runtime/InteractiveToolsFramework/Public/BaseGizmos/GizmoElementCylinder.h
+++ b/Engine/Source/Runtime/InteractiveToolsFramework/Public/BaseGizmos/GizmoElementCylinder.h
@@ -22,7 +22,7 @@ public:
 
 	//~ Begin UGizmoElementBase Interface.
 	INTERACTIVETOOLSFRAMEWORK_API virtual void Render(IToolsContextRenderAPI* RenderAPI, const FRenderTraversalState& RenderState) override;
-	INTERACTIVETOOLSFRAMEWORK_API virtual FInputRayHit LineTrace(const UGizmoViewContext* ViewContext, const FLineTraceTraversalState& LineTraceState, const FVector& RayOrigin, const FVector& RayDirection) override;
+	INTERACTIVETOOLSFRAMEWORK_API virtual FInputRayHit LineTrace(const UGizmoViewContext* ViewContext, const FLineTraceTraversalState& LineTraceState, const FVector& RayOrigin, const FVector& RayDirection, FLineTraceOutput& OutLineTraceOutput) override;
 	//~ End UGizmoElementBase Interface.
 
 	// Location of center of cylinder's base circle.
diff --git a/Engine/Source/Runtime/InteractiveToolsFramework/Public/BaseGizmos/GizmoElementGroup.h b/Engine/Source/Runtime/InteractiveToolsFramework/Public/BaseGizmos/GizmoElementGroup.h
--- a/Engine/Source/Runtime/InteractiveToolsFramework/Public/BaseGizmos/GizmoElementGroup.h
+++ b/Engine/Source/Runtime/InteractiveToolsFramework/Public/BaseGizmos/GizmoElementGroup.h
@@ -24,7 +24,7 @@ public:
 	//~ Begin UGizmoElementBase Interface.
 	INTERACTIVETOOLSFRAMEWORK_API virtual void Render(IToolsContextRenderAPI* RenderAPI, const FRenderTraversalState& RenderState) override;
 	INTERACTIVETOOLSFRAMEWORK_API virtual void DrawHUD(FCanvas* Canvas, IToolsContextRenderAPI* RenderAPI, const FRenderTraversalState& RenderState) override;
-	INTERACTIVETOOLSFRAMEWORK_API virtual FInputRayHit LineTrace(const UGizmoViewContext* ViewContext, const FLineTraceTraversalState& LineTraceState, const FVector& RayOrigin, const FVector& RayDirection) override;
+	INTERACTIVETOOLSFRAMEWORK_API virtual FInputRayHit LineTrace(const UGizmoViewContext* ViewContext, const FLineTraceTraversalState& LineTraceState, const FVector& RayOrigin, const FVector& RayDirection, FLineTraceOutput& OutLineTraceOutput) override;
 
 	// Sets the part identifier for this element and all of its children whose current ID is 0/Default.
 	using UGizmoElementBase::SetPartIdentifier;
diff --git a/Engine/Source/Runtime/InteractiveToolsFramework/Public/BaseGizmos/GizmoElementLineStrip.h b/Engine/Source/Runtime/InteractiveToolsFramework/Public/BaseGizmos/GizmoElementLineStrip.h
--- a/Engine/Source/Runtime/InteractiveToolsFramework/Public/BaseGizmos/GizmoElementLineStrip.h
+++ b/Engine/Source/Runtime/InteractiveToolsFramework/Public/BaseGizmos/GizmoElementLineStrip.h
@@ -24,7 +24,7 @@ public:
 
 	//~ Begin UGizmoElementBase Interface.
 	INTERACTIVETOOLSFRAMEWORK_API virtual void Render(IToolsContextRenderAPI* RenderAPI, const FRenderTraversalState& RenderState) override;
-	INTERACTIVETOOLSFRAMEWORK_API virtual FInputRayHit LineTrace(const UGizmoViewContext* ViewContext, const FLineTraceTraversalState& LineTraceState, const FVector& RayOrigin, const FVector& RayDirection) override;
+	INTERACTIVETOOLSFRAMEWORK_API virtual FInputRayHit LineTrace(const UGizmoViewContext* ViewContext, const FLineTraceTraversalState& LineTraceState, const FVector& RayOrigin, const FVector& RayDirection, FLineTraceOutput& OutLineTraceOutput) override;
 	//~ End UGizmoElementBase Interface.
 
 	// The vertices of the line strip
diff --git a/Engine/Source/Runtime/InteractiveToolsFramework/Public/BaseGizmos/GizmoElementRectangle.h b/Engine/Source/Runtime/InteractiveToolsFramework/Public/BaseGizmos/GizmoElementRectangle.h
--- a/Engine/Source/Runtime/InteractiveToolsFramework/Public/BaseGizmos/GizmoElementRectangle.h
+++ b/Engine/Source/Runtime/InteractiveToolsFramework/Public/BaseGizmos/GizmoElementRectangle.h
@@ -24,7 +24,7 @@ public:
 
 	//~ Begin UGizmoElementBase Interface.
 	INTERACTIVETOOLSFRAMEWORK_API virtual void Render(IToolsContextRenderAPI* RenderAPI, const FRenderTraversalState& RenderState) override;
-	INTERACTIVETOOLSFRAMEWORK_API virtual FInputRayHit LineTrace(const UGizmoViewContext* ViewContext, const FLineTraceTraversalState& LineTraceState, const FVector& RayOrigin, const FVector& RayDirection) override;
+	INTERACTIVETOOLSFRAMEWORK_API virtual FInputRayHit LineTrace(const UGizmoViewContext* ViewContext, const FLineTraceTraversalState& LineTraceState, const FVector& RayOrigin, const FVector& RayDirection, FLineTraceOutput& OutLineTraceOutput) override;
 	//~ End UGizmoElementBase Interface.
 
 	// Location of rectangle center
diff --git a/Engine/Source/Runtime/InteractiveToolsFramework/Public/BaseGizmos/GizmoElementSphere.h b/Engine/Source/Runtime/InteractiveToolsFramework/Public/BaseGizmos/GizmoElementSphere.h
--- a/Engine/Source/Runtime/InteractiveToolsFramework/Public/BaseGizmos/GizmoElementSphere.h
+++ b/Engine/Source/Runtime/InteractiveToolsFramework/Public/BaseGizmos/GizmoElementSphere.h
@@ -22,7 +22,7 @@ public:
 
 	//~ Begin UGizmoElementBase Interface.
 	INTERACTIVETOOLSFRAMEWORK_API virtual void Render(IToolsContextRenderAPI* RenderAPI, const FRenderTraversalState& RenderState) override;
-	INTERACTIVETOOLSFRAMEWORK_API virtual FInputRayHit LineTrace(const UGizmoViewContext* ViewContext, const FLineTraceTraversalState& LineTraceState, const FVector& RayOrigin, const FVector& RayDirection) override;
+	INTERACTIVETOOLSFRAMEWORK_API virtual FInputRayHit LineTrace(const UGizmoViewContext* ViewContext, const FLineTraceTraversalState& LineTraceState, const FVector& RayOrigin, const FVector& RayDirection, FLineTraceOutput& OutLineTraceOutput) override;
 	//~ End UGizmoElementBase Interface.
 
 	// Location of center of sphere
diff --git a/Engine/Source/Runtime/InteractiveToolsFramework/Public/BaseGizmos/GizmoElementTorus.h b/Engine/Source/Runtime/InteractiveToolsFramework/Public/BaseGizmos/GizmoElementTorus.h
--- a/Engine/Source/Runtime/InteractiveToolsFramework/Public/BaseGizmos/GizmoElementTorus.h
+++ b/Engine/Source/Runtime/InteractiveToolsFramework/Public/BaseGizmos/GizmoElementTorus.h
@@ -32,7 +32,7 @@ public:
 	// lies, then determining a hit point closest to the linear circle defined by torus center and torus outer radius.
 	// If the torus lies at a glancing angle, the ray-torus intersection is performed against cylinders approximating
 	// the shape of the torus.
-	INTERACTIVETOOLSFRAMEWORK_API virtual FInputRayHit LineTrace(const UGizmoViewContext* ViewContext, const FLineTraceTraversalState& LineTraceState, const FVector& RayOrigin, const FVector& RayDirection) override;
+	INTERACTIVETOOLSFRAMEWORK_API virtual FInputRayHit LineTrace(const UGizmoViewContext* ViewContext, const FLineTraceTraversalState& LineTraceState, const FVector& RayOrigin, const FVector& RayDirection, FLineTraceOutput& OutLineTraceOutput) override;
 	//~ End UGizmoElementBase Interface.
 
 	// Inner circles radius.
diff --git a/Engine/Source/Runtime/InteractiveToolsFramework/Public/BaseGizmos/GizmoElementTriangleList.h b/Engine/Source/Runtime/InteractiveToolsFramework/Public/BaseGizmos/GizmoElementTriangleList.h
--- a/Engine/Source/Runtime/InteractiveToolsFramework/Public/BaseGizmos/GizmoElementTriangleList.h
+++ b/Engine/Source/Runtime/InteractiveToolsFramework/Public/BaseGizmos/GizmoElementTriangleList.h
@@ -22,7 +22,7 @@ public:
 
 	//~ Begin UGizmoElementBase Interface.
 	INTERACTIVETOOLSFRAMEWORK_API virtual void Render(IToolsContextRenderAPI* RenderAPI, const FRenderTraversalState& RenderState) override;
-	INTERACTIVETOOLSFRAMEWORK_API virtual FInputRayHit LineTrace(const UGizmoViewContext* ViewContext, const FLineTraceTraversalState& LineTraceState, const FVector& RayOrigin, const FVector& RayDirection) override;
+	INTERACTIVETOOLSFRAMEWORK_API virtual FInputRayHit LineTrace(const UGizmoViewContext* ViewContext, const FLineTraceTraversalState& LineTraceState, const FVector& RayOrigin, const FVector& RayDirection, FLineTraceOutput& OutLineTraceOutput) override;
 	//~ End UGizmoElementBase Interface.
 
 	// The vertices of the triangle list, each tuple of 3 forming a triangle
