ModelingTools: Add support for pressure sensitivity to UBaseBrushTool, to be optionally used by subclasses. Also mark some settings as exposed in quick settings.

diff --git a/Engine/Source/Runtime/InteractiveToolsFramework/Private/BaseTools/BaseBrushTool.cpp b/Engine/Source/Runtime/InteractiveToolsFramework/Private/BaseTools/BaseBrushTool.cpp
--- a/Engine/Source/Runtime/InteractiveToolsFramework/Private/BaseTools/BaseBrushTool.cpp
+++ b/Engine/Source/Runtime/InteractiveToolsFramework/Private/BaseTools/BaseBrushTool.cpp
@@ -299,6 +299,10 @@ void UBaseBrushTool::Shutdown(EToolShutdownType ShutdownType)
 	ShutdownBrushStampIndicator();
 }
 
+void UBaseBrushTool::OnTick(float DeltaTime)
+{
+	ActivePressure = GetCurrentDevicePressure();
+}
 
 void UBaseBrushTool::OnPropertyModified(UObject* PropertySet, FProperty* Property)
 {
@@ -391,6 +395,16 @@ void UBaseBrushTool::DecreaseBrushFalloffAction()
 	NotifyOfPropertyChangeByTool(BrushProperties);
 }
 
+double UBaseBrushTool::GetCurrentBrushRadius() const
+{
+	if (BrushProperties && BrushProperties->bToolSupportsPressureSensitivity 
+		&& BrushProperties->bEnablePressureSensitivity)
+	{
+		return CurrentBrushRadius * GetActivePressure();
+	}
+	return CurrentBrushRadius;
+}
+
 void UBaseBrushTool::SetBrushEnabled(bool bIsEnabled)
 {
 	bEnabled = bIsEnabled;
diff --git a/Engine/Source/Runtime/InteractiveToolsFramework/Public/BaseTools/BaseBrushTool.h b/Engine/Source/Runtime/InteractiveToolsFramework/Public/BaseTools/BaseBrushTool.h
--- a/Engine/Source/Runtime/InteractiveToolsFramework/Public/BaseTools/BaseBrushTool.h
+++ b/Engine/Source/Runtime/InteractiveToolsFramework/Public/BaseTools/BaseBrushTool.h
@@ -35,12 +35,29 @@ public:
 		DisplayName = "Radius", UIMin = "1.0", UIMax = "1000.0", ClampMin = "0.1", ClampMax = "50000.0"))
 	float BrushRadius;
 
+	//~Set by tool and determines whether the bEnablePressureSensitivity flag below is shown and used.
+	//~ All that a subclass of UBaseBrushTool needs to do to support pressure sensitivity is to call
+	//~ the super OnTick, use the provided radius getters where relevant, and set this to true in setup. 
+	//~ Its StylusAPI also should be set via its builder.
+	UPROPERTY()
+	bool bToolSupportsPressureSensitivity = false;
+
+	/** When set, stylus pressure is used to adjust radius */
+	UPROPERTY(EditAnywhere, Category = Brush, meta = (EditCondition = "bToolSupportsPressureSensitivity", EditConditionHides))
+	bool bEnablePressureSensitivity = false;
+
 	/** Strength of the brush (0.0 - 1.0) */
-	UPROPERTY(EditAnywhere, Category = Brush, meta = (DisplayName = "Strength", UIMin = "0.0", UIMax = "1.0", ClampMin = "0.0", ClampMax = "1.0", HideEditConditionToggle, EditConditionHides, EditCondition = "bShowStrength", DisplayPriority = 2))
+	UPROPERTY(EditAnywhere, Category = Brush, meta = (DisplayName = "Strength",
+		UIMin = "0.0", UIMax = "1.0", ClampMin = "0.0", ClampMax = "1.0", 
+		HideEditConditionToggle, EditConditionHides, EditCondition = "bShowStrength", DisplayPriority = 2,
+		ModelingQuickEdit, ModelingQuickSettings))
 	float BrushStrength;
 
 	/** Amount of falloff to apply (0.0 - 1.0) */
-	UPROPERTY(EditAnywhere, Category = Brush, meta = (DisplayName = "Falloff", UIMin = "0.0", UIMax = "1.0", ClampMin = "0.0", ClampMax = "1.0", HideEditConditionToggle, EditConditionHides, EditCondition="bShowFalloff", DisplayPriority = 3))
+	UPROPERTY(EditAnywhere, Category = Brush, meta = (DisplayName = "Falloff", 
+		UIMin = "0.0", UIMax = "1.0", ClampMin = "0.0", ClampMax = "1.0", 
+		HideEditConditionToggle, EditConditionHides, EditCondition="bShowFalloff", DisplayPriority = 3,
+		ModelingQuickEdit, ModelingQuickSettings))
 	float BrushFalloffAmount;
 
 	/** If false, then BrushStrength will not be shown in DetailsView panels (otherwise no effect) */
@@ -160,6 +177,7 @@ public:
 	// UInteractiveTool
 	INTERACTIVETOOLSFRAMEWORK_API virtual void Setup() override;
 	INTERACTIVETOOLSFRAMEWORK_API virtual void Shutdown(EToolShutdownType ShutdownType) override;
+	INTERACTIVETOOLSFRAMEWORK_API virtual void OnTick(float DeltaTime) override;
 	INTERACTIVETOOLSFRAMEWORK_API virtual void Render(IToolsContextRenderAPI* RenderAPI) override;
 	INTERACTIVETOOLSFRAMEWORK_API virtual void DrawHUD( FCanvas* Canvas, IToolsContextRenderAPI* RenderAPI ) override;
 	
@@ -208,8 +226,8 @@ public:
 	virtual bool IsInBrushStroke() const { return bInBrushStroke; }
 	PRAGMA_ENABLE_DEPRECATION_WARNINGS
 
-	virtual double GetCurrentBrushRadius() const { return CurrentBrushRadius; }
-	virtual double GetCurrentBrushRadiusLocal() const { return CurrentBrushRadius * WorldToLocalScale; }
+	INTERACTIVETOOLSFRAMEWORK_API virtual double GetCurrentBrushRadius() const;
+	virtual double GetCurrentBrushRadiusLocal() const { return GetCurrentBrushRadius() * WorldToLocalScale; }
 
 	INTERACTIVETOOLSFRAMEWORK_API void SetBrushEnabled(bool bIsEnabled);
 
@@ -223,6 +241,7 @@ protected:
 	 */
 	virtual double EstimateMaximumTargetDimension() { return 100.0; }
 
+	virtual double GetActivePressure() const { return ActivePressure; }
 protected:
 	TInterval<float> BrushRelativeSizeRange;
 	double CurrentBrushRadius;
@@ -252,6 +271,6 @@ protected:
 	TWeakObjectPtr<ULocalTwoAxisPropertyEditInputBehavior> BrushEditBehavior;
 
 private:
-	
+	double ActivePressure = 1.0;
 	bool bEnabled = true;
 };
