IInteractiveToolShutdownQueryAPI: Deprecate internally defined EShutdownReason, move EShutdownReason (named EToolShutdownReason) to InteractiveTool.h

In order to augment the tool manager's interface dealing with preferred shutdown types, we need to decouple this enum from the interface it currently lives in.

diff --git a/Engine/Source/Runtime/InteractiveToolsFramework/Private/InteractiveToolManager.cpp b/Engine/Source/Runtime/InteractiveToolsFramework/Private/InteractiveToolManager.cpp
--- a/Engine/Source/Runtime/InteractiveToolsFramework/Private/InteractiveToolManager.cpp
+++ b/Engine/Source/Runtime/InteractiveToolsFramework/Private/InteractiveToolManager.cpp
@@ -179,8 +179,7 @@ bool UInteractiveToolManager::ActivateTool(EToolSide Side)
 		{
 			if (IInteractiveToolShutdownQueryAPI* ShutdownQueryAPI = Cast<IInteractiveToolShutdownQueryAPI>(ActiveLeftTool))
 			{
-				ShutdownType = ShutdownQueryAPI->GetPreferredShutdownType(
-					IInteractiveToolShutdownQueryAPI::EShutdownReason::SwitchTool, ShutdownType);
+				ShutdownType = ShutdownQueryAPI->GetPreferredShutdownType(EToolShutdownReason::SwitchTool, ShutdownType);
 			}
 		}
 
diff --git a/Engine/Source/Runtime/InteractiveToolsFramework/Public/InteractiveTool.h b/Engine/Source/Runtime/InteractiveToolsFramework/Public/InteractiveTool.h
--- a/Engine/Source/Runtime/InteractiveToolsFramework/Public/InteractiveTool.h
+++ b/Engine/Source/Runtime/InteractiveToolsFramework/Public/InteractiveTool.h
@@ -51,6 +51,13 @@ enum class EToolShutdownType : uint8
 	Cancel = 2
 };
 
+/** Used in various ways by the tool manager and some optional tool interfaces to make shutdown responses configurable. */
+enum class EToolShutdownReason : uint8
+{
+	/** Another tool is being activated without the user having explicitly shut down this one. */
+	SwitchTool = 1
+};
+
 /**
  * FInteractiveToolInfo provides information about a tool (name, tooltip, etc)
  */
diff --git a/Engine/Source/Runtime/InteractiveToolsFramework/Public/InteractiveToolQueryInterfaces.h b/Engine/Source/Runtime/InteractiveToolsFramework/Public/InteractiveToolQueryInterfaces.h
--- a/Engine/Source/Runtime/InteractiveToolsFramework/Public/InteractiveToolQueryInterfaces.h
+++ b/Engine/Source/Runtime/InteractiveToolsFramework/Public/InteractiveToolQueryInterfaces.h
@@ -3,7 +3,7 @@
 #pragma once
 
 #include "CoreMinimal.h"
-#include "InteractiveTool.h" // EToolShutdownType
+#include "InteractiveTool.h" // EToolShutdownType, EToolShutdownReason
 #include "UObject/Interface.h"
 #include "InteractiveToolQueryInterfaces.generated.h"
 
@@ -153,16 +153,11 @@ class IInteractiveToolShutdownQueryAPI
 {
 	GENERATED_BODY()
 public:
-	enum class EShutdownReason
+
+	enum class UE_DEPRECATED(5.7, "EShutdownReason has moved to InteractiveTool.h (under the name EToolShutdownReason).") EShutdownReason
 	{
 		// Another tool is being activated without the user having explicitly shut down this one.
 		SwitchTool
-
-		//~ We'll add something along these lines once we start querying in other shutdown situations:
-		//~ // The ESC key (or some equivalent) was hit to exit the tool
-		//~ , Escape
-		//~ // The mode was shut down
-		//~ , ModeExit
 	};
 
 	/**
@@ -172,8 +167,31 @@ public:
 	 *   chance to customize the behavior.
 	 * @return Shutdown type that the tool would prefer that the host use in this situation.
 	 */
+	UE_DEPRECATED(5.7, "Prefer the GetPreferredShutdownType variant which uses EToolShutdownReason.")
+	PRAGMA_DISABLE_DEPRECATION_WARNINGS
 	virtual EToolShutdownType GetPreferredShutdownType(EShutdownReason ShutdownReason, EToolShutdownType StandardShutdownType) const
+	PRAGMA_ENABLE_DEPRECATION_WARNINGS
+	{
+		return StandardShutdownType;
+	}
+	
+	/**
+	 * Given a shutdown situation, tells what kind of shutdown method the tool might prefer.
+	 * @param ShutdownReason Information about the shutdown type
+	 * @param StandardShutdownType Shutdown type that the host would use if it weren't giving the tool a
+	 *   chance to customize the behavior.
+	 * @return Shutdown type that the tool would prefer that the host use in this situation.
+	 */
+	virtual EToolShutdownType GetPreferredShutdownType(EToolShutdownReason ShutdownReason, EToolShutdownType StandardShutdownType) const
 	{
+		if (ShutdownReason == EToolShutdownReason::SwitchTool)
+		{
+			// todo: remove this special case when removing the deprecated variant of GetPreferredShutdownType.
+			PRAGMA_DISABLE_DEPRECATION_WARNINGS
+			return GetPreferredShutdownType(EShutdownReason::SwitchTool, StandardShutdownType);
+			PRAGMA_ENABLE_DEPRECATION_WARNINGS
+		}
+	
 		return StandardShutdownType;
 	}
 
