[ITF Input Router] Make sure HandleCapturedMouseInput only handles mouse inputs when calling CollectWantsCapture

CollectWantsCapture is called by HandleCapturedMouseInput, in order to allow mouse buttons to be treated as modifiers, when needed.
This change makes sure that happens only for mouse inputs, and not for keyboard inputs.

diff --git a/Engine/Source/Runtime/InteractiveToolsFramework/Private/InputRouter.cpp b/Engine/Source/Runtime/InteractiveToolsFramework/Private/InputRouter.cpp
--- a/Engine/Source/Runtime/InteractiveToolsFramework/Private/InputRouter.cpp
+++ b/Engine/Source/Runtime/InteractiveToolsFramework/Private/InputRouter.cpp
@@ -273,10 +273,13 @@ void UInputRouter::CheckForMouseCaptures(const FInputDeviceState& Input)
 
 void UInputRouter::HandleCapturedMouseInput(const FInputDeviceState& Input)
 {
-	// Calling this to make sure behaviors treating mouse buttons as modifiers are always updated,
-	// even if there is some drag-based behavior currently capturing
-	TArray<FInputCaptureRequest> CaptureRequests;
-	ActiveInputBehaviors->CollectWantsCapture(Input, CaptureRequests);
+	if (Input.IsFromDevice(EInputDevices::Mouse))
+	{
+		// Calling this to make sure behaviors treating mouse buttons as modifiers are always updated,
+		// even if there is some drag-based behavior currently capturing
+		TArray<FInputCaptureRequest> CaptureRequests;
+		ActiveInputBehaviors->CollectWantsCapture(Input, CaptureRequests);
+	}
 
 	if (ActiveLeftCapture == nullptr)
 	{
