[ITF] Fixed outline scale on Arc Gizmo Element

diff --git a/Engine/Source/Runtime/InteractiveToolsFramework/Private/BaseGizmos/GizmoElementArc.cpp b/Engine/Source/Runtime/InteractiveToolsFramework/Private/BaseGizmos/GizmoElementArc.cpp
--- a/Engine/Source/Runtime/InteractiveToolsFramework/Private/BaseGizmos/GizmoElementArc.cpp
+++ b/Engine/Source/Runtime/InteractiveToolsFramework/Private/BaseGizmos/GizmoElementArc.cpp
@@ -175,7 +175,8 @@ void UGizmoElementArc::Render(IToolsContextRenderAPI* RenderAPI, const FRenderTr
 	const double WorldInnerRadius = InnerRadius * CurrentRenderState.LocalToWorldTransform.GetScale3D().X;
 	const FColor VertexColor = CurrentRenderState.GetCurrentVertexColor().ToFColor(true);
 	const FColor LineColor = CurrentRenderState.GetCurrentLineColor().ToFColor(true);
-	const double CurrentLineThickness = GetCurrentLineThickness(View->IsPerspectiveProjection(), View->FOV);
+	const double CurrentLineThickness = GetCurrentLineThickness(View->IsPerspectiveProjection(), View->FOV)
+										* CurrentRenderState.LocalToWorldTransform.GetScale3D().X;
 
 	FPrimitiveDrawInterface* PDI = RenderAPI->GetPrimitiveDrawInterface();
 
