[ITF] Added fixed dash count limit to prevent overflow

diff --git a/Engine/Source/Runtime/InteractiveToolsFramework/Private/BaseGizmos/GizmoElementCylinder.cpp b/Engine/Source/Runtime/InteractiveToolsFramework/Private/BaseGizmos/GizmoElementCylinder.cpp
--- a/Engine/Source/Runtime/InteractiveToolsFramework/Private/BaseGizmos/GizmoElementCylinder.cpp
+++ b/Engine/Source/Runtime/InteractiveToolsFramework/Private/BaseGizmos/GizmoElementCylinder.cpp
@@ -54,9 +54,11 @@ void UGizmoElementCylinder::Render(IToolsContextRenderAPI* RenderAPI, const FRen
 			{
 				// Get a scale factor based on how many dashes with gaps (and an end dash) can actually fit on this line
 				float DashLengthScale = 0.0f;
+				
+				constexpr int32 MaxNumDashes = 512; // Arbitrary maximum to prevent overflow in the case of very long cylinders
 
 				// Number based on sequence of dash, gap pairs, and last dash
-				const int32 NumDashes = GizmoElementCylinderLocals::GetDashNumAndScale(FMath::Max(0.0f, Height - DashLength), DashLength, DashGapLength, DashLengthScale);
+				const int32 NumDashes = FMath::Min(MaxNumDashes, GizmoElementCylinderLocals::GetDashNumAndScale(FMath::Max(0.0f, Height - DashLength), DashLength, DashGapLength, DashLengthScale));
 
 				const float AdjustedDashLength = DashLength * DashLengthScale;
 				const float AdjustDashGapLength = DashGapLength * DashLengthScale;
