[ITF] Inject DPI into LineTraceState

diff --git a/Engine/Source/Runtime/InteractiveToolsFramework/Public/BaseGizmos/GizmoElementBase.h b/Engine/Source/Runtime/InteractiveToolsFramework/Public/BaseGizmos/GizmoElementBase.h
--- a/Engine/Source/Runtime/InteractiveToolsFramework/Public/BaseGizmos/GizmoElementBase.h
+++ b/Engine/Source/Runtime/InteractiveToolsFramework/Public/BaseGizmos/GizmoElementBase.h
@@ -54,7 +54,7 @@ public:
 		// Note: non-uniform scale is not supported and the X scale element will be used for uniform scaling.
 		FTransform LocalToWorldTransform;
 
-		// Pixel to world scale. To account for DPI scale, @see IToolsContextRenderAPI::ViewCameraState.
+		// Pixel to world scale. Also accounts for DPI scale, @see IToolsContextRenderAPI::ViewCameraState.
 		double PixelToWorldScale = 1.0;
 
 		// Interact state, if not equal to none, overrides the element's interact state 
@@ -108,9 +108,12 @@ public:
 		void Initialize(const UGizmoViewContext* InGizmoViewContext, FTransform InTransform)
 		{
 			check(InGizmoViewContext);
+
 			LocalToWorldTransform = InTransform;
-			PixelToWorldScale = UE::GizmoRenderingUtil::CalculateLocalPixelToWorldScale(InGizmoViewContext, InTransform.GetLocation());
 			bIsPerspectiveProjection = InGizmoViewContext->IsPerspectiveProjection();
+
+			constexpr bool bUseDPIScaleFromViewContext = true;
+			PixelToWorldScale = UE::GizmoRenderingUtil::CalculateLocalPixelToWorldScale(InGizmoViewContext, InTransform.GetLocation(), bUseDPIScaleFromViewContext);
 		}
 	};
 
