ITF - Add GetPreferredShutdownType which determines the preferred shutdown type based on the state of the tool manager and the active tool itself.

diff --git a/Engine/Source/Runtime/InteractiveToolsFramework/Private/InteractiveToolManager.cpp b/Engine/Source/Runtime/InteractiveToolsFramework/Private/InteractiveToolManager.cpp
--- a/Engine/Source/Runtime/InteractiveToolsFramework/Private/InteractiveToolManager.cpp
+++ b/Engine/Source/Runtime/InteractiveToolsFramework/Private/InteractiveToolManager.cpp
@@ -155,35 +155,7 @@ bool UInteractiveToolManager::ActivateTool(EToolSide Side)
 
 	if (ActiveLeftTool != nullptr)
 	{
-		// Decide how to shut down the currently active tool based on ToolSwitchMode
-		EToolShutdownType ShutdownType = EToolShutdownType::Completed;
-		switch (ToolSwitchMode)
-		{
-		case EToolManagerToolSwitchMode::AcceptIfAble:
-		case EToolManagerToolSwitchMode::CustomizableAcceptIfAble:
-			ShutdownType = 
-				CanAcceptActiveTool(EToolSide::Left) ? EToolShutdownType::Accept : 
-				CanCancelActiveTool(EToolSide::Left) ? EToolShutdownType::Cancel : EToolShutdownType::Completed;
-			break;
-		case EToolManagerToolSwitchMode::CancelIfAble:
-		case EToolManagerToolSwitchMode::CustomizableCancelIfAble:
-			ShutdownType = CanCancelActiveTool(EToolSide::Left) ? EToolShutdownType::Cancel : EToolShutdownType::Completed;
-			break;
-		default:
-			ensure(false);
-		}
-
-		// If the tool switch mode is set to be customizable, see if the tool wants to customize it
-		if (ToolSwitchMode == EToolManagerToolSwitchMode::CustomizableAcceptIfAble
-			|| ToolSwitchMode == EToolManagerToolSwitchMode::CustomizableCancelIfAble)
-		{
-			if (IInteractiveToolShutdownQueryAPI* ShutdownQueryAPI = Cast<IInteractiveToolShutdownQueryAPI>(ActiveLeftTool))
-			{
-				ShutdownType = ShutdownQueryAPI->GetPreferredShutdownType(EToolShutdownReason::SwitchTool, ShutdownType);
-			}
-		}
-
-		DeactivateTool(EToolSide::Left, ShutdownType);
+		DeactivateTool(EToolSide::Left, GetPreferredShutdownType(EToolSide::Left, EToolShutdownReason::SwitchTool));
 	}
 
 	if (ActiveLeftBuilder == nullptr || ActivateToolInternal(Side) == false)
@@ -287,6 +259,39 @@ void UInteractiveToolManager::DeactivateTool(EToolSide Side, EToolShutdownType S
 	}
 }
 
+EToolShutdownType UInteractiveToolManager::GetPreferredShutdownType(EToolSide Side, EToolShutdownReason ShutdownReason)
+{
+	// Decide how to shut down the currently active tool based on ToolSwitchMode
+	EToolShutdownType ShutdownType = EToolShutdownType::Completed;
+	switch (ToolSwitchMode)
+	{
+	case EToolManagerToolSwitchMode::AcceptIfAble:
+	case EToolManagerToolSwitchMode::CustomizableAcceptIfAble:
+		ShutdownType = 
+			CanAcceptActiveTool(Side) ? EToolShutdownType::Accept : 
+			CanCancelActiveTool(Side) ? EToolShutdownType::Cancel : EToolShutdownType::Completed;
+		break;
+	case EToolManagerToolSwitchMode::CancelIfAble:
+	case EToolManagerToolSwitchMode::CustomizableCancelIfAble:
+		ShutdownType = CanCancelActiveTool(Side) ? EToolShutdownType::Cancel : EToolShutdownType::Completed;
+		break;
+	default:
+		ensure(false);
+	}
+
+	// If the tool switch mode is set to be customizable, see if the tool wants to customize it
+	if (ToolSwitchMode == EToolManagerToolSwitchMode::CustomizableAcceptIfAble
+		|| ToolSwitchMode == EToolManagerToolSwitchMode::CustomizableCancelIfAble)
+	{
+		if (IInteractiveToolShutdownQueryAPI* ShutdownQueryAPI = Cast<IInteractiveToolShutdownQueryAPI>(ActiveLeftTool))
+		{
+			ShutdownType = ShutdownQueryAPI->GetPreferredShutdownType(
+				ShutdownReason, ShutdownType);
+		}
+	}
+	
+	return ShutdownType;
+}
 
 void UInteractiveToolManager::DeactivateToolInternal(EToolSide Side, EToolShutdownType ShutdownType)
 {
diff --git a/Engine/Source/Runtime/InteractiveToolsFramework/Public/InteractiveTool.h b/Engine/Source/Runtime/InteractiveToolsFramework/Public/InteractiveTool.h
--- a/Engine/Source/Runtime/InteractiveToolsFramework/Public/InteractiveTool.h
+++ b/Engine/Source/Runtime/InteractiveToolsFramework/Public/InteractiveTool.h
@@ -54,6 +54,8 @@ enum class EToolShutdownType : uint8
 /** Used in various ways by the tool manager and some optional tool interfaces to make shutdown responses configurable. */
 enum class EToolShutdownReason : uint8
 {
+	/** The reason for the shutdown was left unspecified. */
+	Unspecified = 0,
 	/** Another tool is being activated without the user having explicitly shut down this one. */
 	SwitchTool = 1
 };
diff --git a/Engine/Source/Runtime/InteractiveToolsFramework/Public/InteractiveToolManager.h b/Engine/Source/Runtime/InteractiveToolsFramework/Public/InteractiveToolManager.h
--- a/Engine/Source/Runtime/InteractiveToolsFramework/Public/InteractiveToolManager.h
+++ b/Engine/Source/Runtime/InteractiveToolsFramework/Public/InteractiveToolManager.h
@@ -223,6 +223,11 @@ public:
 	 */
 	INTERACTIVETOOLSFRAMEWORK_API virtual void DeactivateTool(EToolSide Side, EToolShutdownType ShutdownType);
 
+	/**
+	 * Determines the preferred shutdown type for the active tool based on the tool's type and the tool manager's configuration.
+	 */
+	INTERACTIVETOOLSFRAMEWORK_API EToolShutdownType GetPreferredShutdownType(EToolSide Side,
+		EToolShutdownReason ShutdownReason = EToolShutdownReason::Unspecified);
 
 	/**
 	 * Configure how tool changes emit change events. See EToolChangeTrackingMode for details.
