Modeling Tools: update tags for BaseBrushTool

diff --git a/Engine/Source/Runtime/InteractiveToolsFramework/Public/BaseTools/BaseBrushTool.h b/Engine/Source/Runtime/InteractiveToolsFramework/Public/BaseTools/BaseBrushTool.h
--- a/Engine/Source/Runtime/InteractiveToolsFramework/Public/BaseTools/BaseBrushTool.h
+++ b/Engine/Source/Runtime/InteractiveToolsFramework/Public/BaseTools/BaseBrushTool.h
@@ -23,7 +23,7 @@ public:
 	INTERACTIVETOOLSFRAMEWORK_API UBrushBaseProperties();
 
 	/** Relative size of brush */
-	UPROPERTY(EditAnywhere, Category = Brush, meta = (DisplayName = "Size", UIMin = "0.0", UIMax = "1.0", ClampMin = "0.0", ClampMax = "10.0", DisplayPriority = 1, HideEditConditionToggle, EditCondition = "bSpecifyRadius == false"))
+	UPROPERTY(EditAnywhere, Category = Brush, meta = (DisplayName = "Size", UIMin = "0.0", UIMax = "1.0", ClampMin = "0.0", ClampMax = "10.0", DisplayPriority = 1, HideEditConditionToggle, EditCondition = "bSpecifyRadius == false", ModelingQuickSettings = 100))
 	float BrushSize;
 
 	/** If true, ignore relative Brush Size and use explicit world Radius */
@@ -50,14 +50,14 @@ public:
 	UPROPERTY(EditAnywhere, Category = Brush, meta = (DisplayName = "Strength",
 		UIMin = "0.0", UIMax = "1.0", ClampMin = "0.0", ClampMax = "1.0", 
 		HideEditConditionToggle, EditConditionHides, EditCondition = "bShowStrength", DisplayPriority = 2,
-		ModelingQuickEdit, ModelingQuickSettings))
+		ModelingQuickEdit, ModelingQuickSettings = 200))
 	float BrushStrength;
 
 	/** Amount of falloff to apply (0.0 - 1.0) */
 	UPROPERTY(EditAnywhere, Category = Brush, meta = (DisplayName = "Falloff", 
 		UIMin = "0.0", UIMax = "1.0", ClampMin = "0.0", ClampMax = "1.0", 
 		HideEditConditionToggle, EditConditionHides, EditCondition="bShowFalloff", DisplayPriority = 3,
-		ModelingQuickEdit, ModelingQuickSettings))
+		ModelingQuickEdit, ModelingQuickSettings = 300))
 	float BrushFalloffAmount;
 
 	/** If false, then BrushStrength will not be shown in DetailsView panels (otherwise no effect) */
