[ITF] Implement Capture Stealing in ITF Input Router

- Adds capture stealing logic to Input Router for Mouse and Keyboard
- InputBeavior can now specify if they will allow capture stealing

[RN] minor editor

diff --git a/Engine/Source/Runtime/InteractiveToolsFramework/Private/InputBehavior.cpp b/Engine/Source/Runtime/InteractiveToolsFramework/Private/InputBehavior.cpp
--- a/Engine/Source/Runtime/InteractiveToolsFramework/Private/InputBehavior.cpp
+++ b/Engine/Source/Runtime/InteractiveToolsFramework/Private/InputBehavior.cpp
@@ -93,3 +93,13 @@ bool UInputBehavior::WantsForceEndCapture()
 {
 	return false;
 }
+
+bool UInputBehavior::AllowsCaptureStealing()
+{
+	return CaptureStealing.IsSet() ? CaptureStealing.Get() : false;
+}
+
+void UInputBehavior::SetAllowsCaptureStealing(TAttribute<bool> InAllowsCaptureStealing)
+{
+	CaptureStealing = InAllowsCaptureStealing;
+}
diff --git a/Engine/Source/Runtime/InteractiveToolsFramework/Private/InputRouter.cpp b/Engine/Source/Runtime/InteractiveToolsFramework/Private/InputRouter.cpp
--- a/Engine/Source/Runtime/InteractiveToolsFramework/Private/InputRouter.cpp
+++ b/Engine/Source/Runtime/InteractiveToolsFramework/Private/InputRouter.cpp
@@ -169,20 +169,61 @@ void UInputRouter::HandleCapturedKeyboardInput(const FInputDeviceState& Input)
 		return;
 	}
 
-	FInputCaptureUpdate Result =
-		ActiveKeyboardCapture->UpdateCapture(Input, ActiveKeyboardCaptureData);
-
-	if (Result.State == EInputCaptureState::End)
+	// Perform Capture Stealing if needed
+	bool bStoleCapture = false;
+	if (ActiveKeyboardCapture->AllowsCaptureStealing())
 	{
-		ActiveKeyboardCapture = nullptr;
-		ActiveKeyboardCaptureOwner = nullptr;
-		ActiveKeyboardCaptureData = FInputCaptureData();
+		TArray<FInputCaptureRequest> CaptureRequests;
+		ActiveInputBehaviors->CollectWantsCapture(Input, CaptureRequests);
+
+		if (!CaptureRequests.IsEmpty())
+		{
+			CaptureRequests.StableSort();
+
+			const FInputCaptureRequest& TopRequest = CaptureRequests[0];
+			UInputBehavior* const TopCapture = TopRequest.Source;
+
+			if (TopCapture != ActiveKeyboardCapture)
+			{
+				if (TopCapture->GetPriority() < ActiveKeyboardCapture->GetPriority())
+				{
+					ActiveKeyboardCapture->ForceEndCapture(ActiveKeyboardCaptureData);
+
+					ActiveKeyboardCapture = nullptr;
+					ActiveKeyboardCaptureOwner = nullptr;
+					ActiveKeyboardCaptureData = FInputCaptureData();
+
+					FInputCaptureUpdate Result = TopRequest.Source->BeginCapture(Input, EInputCaptureSide::Left);
+					if (Result.State == EInputCaptureState::Begin)
+					{
+						ActiveKeyboardCapture = Result.Source;
+						ActiveKeyboardCaptureOwner = TopRequest.Owner;
+						ActiveKeyboardCaptureData = Result.Data;
+					}
+
+					bStoleCapture = true;
+				}
+			}
+		}
 	}
-	else if (Result.State != EInputCaptureState::Continue)
+
+	if (!bStoleCapture)
 	{
-		TransactionsAPI->DisplayMessage(LOCTEXT("HandleCapturedKeyboardInputMessage", "UInteractiveToolManager::HandleCapturedKeyboardInput - unexpected capture state!"), EToolMessageLevel::Internal);
-	}
+		FInputCaptureUpdate Result =
+		ActiveKeyboardCapture->UpdateCapture(Input, ActiveKeyboardCaptureData);
 
+		if (Result.State == EInputCaptureState::End)
+		{
+			ActiveKeyboardCapture = nullptr;
+			ActiveKeyboardCaptureOwner = nullptr;
+			ActiveKeyboardCaptureData = FInputCaptureData();
+		}
+		else if (Result.State != EInputCaptureState::Continue)
+		{
+			TransactionsAPI->DisplayMessage(LOCTEXT("HandleCapturedKeyboardInputMessage", "UInteractiveToolManager::HandleCapturedKeyboardInput - unexpected capture state!"), EToolMessageLevel::Internal);
+		}
+	}
+	
 	if (bAutoInvalidateOnCapture)
 	{
 		TransactionsAPI->PostInvalidation();
@@ -286,22 +327,63 @@ void UInputRouter::HandleCapturedMouseInput(const FInputDeviceState& Input)
 		return;
 	}
 
-	// have active capture - give it this event
+	// Perform Capture Stealing if needed
+	bool bStoleCapture = false;
+	if (ActiveLeftCapture->AllowsCaptureStealing())
+	{
+		TArray<FInputCaptureRequest> CaptureRequests;
+		ActiveInputBehaviors->CollectWantsCapture(Input, CaptureRequests);
 
-	FInputCaptureUpdate Result =
-		ActiveLeftCapture->UpdateCapture(Input, ActiveLeftCaptureData);
+		if (!CaptureRequests.IsEmpty())
+		{
+			CaptureRequests.StableSort();
 
-	if (Result.State == EInputCaptureState::End)
-	{
-		ActiveLeftCapture = nullptr;
-		ActiveLeftCaptureOwner = nullptr;
-		ActiveLeftCaptureData = FInputCaptureData();
+			const FInputCaptureRequest& TopRequest = CaptureRequests[0];
+			UInputBehavior* const TopCapture = TopRequest.Source;
+
+			if (TopCapture != ActiveLeftCapture)
+			{
+				if (TopCapture->GetPriority() < ActiveLeftCapture->GetPriority())
+				{
+					ActiveLeftCapture->ForceEndCapture(ActiveLeftCaptureData);
+					
+					ActiveLeftCapture = nullptr;
+					ActiveLeftCaptureOwner = nullptr;
+					ActiveLeftCaptureData = FInputCaptureData();
+
+					FInputCaptureUpdate Result = TopRequest.Source->BeginCapture(Input, EInputCaptureSide::Left);
+					if (Result.State == EInputCaptureState::Begin)
+					{
+						ActiveLeftCapture = Result.Source;
+						ActiveLeftCaptureOwner = TopRequest.Owner;
+						ActiveLeftCaptureData = Result.Data;
+					}
+
+					bStoleCapture = true;
+				}
+			}
+		}
 	}
-	else if (Result.State != EInputCaptureState::Continue)
+
+	// have active capture - give it this event
+
+	if (!bStoleCapture)
 	{
-		TransactionsAPI->DisplayMessage(LOCTEXT("HandleCapturedMouseInputMessage", "UInteractiveToolManager::HandleCapturedMouseInput - unexpected capture state!"), EToolMessageLevel::Internal);
-	}
+		FInputCaptureUpdate Result =
+		ActiveLeftCapture->UpdateCapture(Input, ActiveLeftCaptureData);
 
+		if (Result.State == EInputCaptureState::End)
+		{
+			ActiveLeftCapture = nullptr;
+			ActiveLeftCaptureOwner = nullptr;
+			ActiveLeftCaptureData = FInputCaptureData();
+		}
+		else if (Result.State != EInputCaptureState::Continue)
+		{
+			TransactionsAPI->DisplayMessage(LOCTEXT("HandleCapturedMouseInputMessage", "UInteractiveToolManager::HandleCapturedMouseInput - unexpected capture state!"), EToolMessageLevel::Internal);
+		}
+	}
+	
 	if (bAutoInvalidateOnCapture)
 	{
 		TransactionsAPI->PostInvalidation();
diff --git a/Engine/Source/Runtime/InteractiveToolsFramework/Public/InputBehavior.h b/Engine/Source/Runtime/InteractiveToolsFramework/Public/InputBehavior.h
--- a/Engine/Source/Runtime/InteractiveToolsFramework/Public/InputBehavior.h
+++ b/Engine/Source/Runtime/InteractiveToolsFramework/Public/InputBehavior.h
@@ -307,7 +307,20 @@ public:
 	 */
 	INTERACTIVETOOLSFRAMEWORK_API virtual bool WantsForceEndCapture();
 
+	/**
+	 * When this function returns true, this behavior allows the Input Router to interrupt capturing and pass it on to a different Behavior with a higher priority.
+	 */
+	INTERACTIVETOOLSFRAMEWORK_API virtual bool AllowsCaptureStealing();
+
+	/**
+	 * Call to set the default capture stealing attribute.
+	 */
+	INTERACTIVETOOLSFRAMEWORK_API void SetAllowsCaptureStealing(TAttribute<bool> InAllowsCaptureStealing);
+
 protected:
 	/** priority returned by GetPriority() */
 	FInputCapturePriority DefaultPriority;
+
+	/** value returned by AllowsCaptureStealing*/
+	TAttribute<bool> CaptureStealing;
 };
