[ITF] PixelToWorldScale clamped to world size vs. arbitrary number

diff --git a/Engine/Source/Runtime/InteractiveToolsFramework/Private/BaseGizmos/GizmoRenderingUtil.cpp b/Engine/Source/Runtime/InteractiveToolsFramework/Private/BaseGizmos/GizmoRenderingUtil.cpp
--- a/Engine/Source/Runtime/InteractiveToolsFramework/Private/BaseGizmos/GizmoRenderingUtil.cpp
+++ b/Engine/Source/Runtime/InteractiveToolsFramework/Private/BaseGizmos/GizmoRenderingUtil.cpp
@@ -147,8 +147,8 @@ float UE::GizmoRenderingUtil::CalculateLocalPixelToWorldScale(
 	FVector OffsetPointWorld = Location + OffsetDelta * View->GetViewRight() + OffsetDelta * View->GetViewUp();
 	FVector2D PixelB = WorldToPixelDouble(View, OffsetPointWorld);
 
-	constexpr double MinPixelToWorldScale = UE_SMALL_NUMBER;
-	constexpr double MaxPixelToWorldScale = 100.0; // Arbitrary large value to prevent division by zero or very small scale
+	constexpr double MinPixelToWorldScale = UE_DOUBLE_SMALL_NUMBER;
+	constexpr double MaxPixelToWorldScale = UE_LARGE_HALF_WORLD_MAX; // Arbitrary large value to prevent division by zero or very small scale
 
 	double PixelDeltaSqr = VectorDifferenceSqr(PixelA, PixelB);
 	double WorldDeltaSqr = VectorDifferenceSqr(Location, OffsetPointWorld);
