[ITF] Fixed gizmo element hit priority - negative should ignore the "exact hit" exception

diff --git a/Engine/Source/Runtime/InteractiveToolsFramework/Private/BaseGizmos/GizmoElementGroup.cpp b/Engine/Source/Runtime/InteractiveToolsFramework/Private/BaseGizmos/GizmoElementGroup.cpp
--- a/Engine/Source/Runtime/InteractiveToolsFramework/Private/BaseGizmos/GizmoElementGroup.cpp
+++ b/Engine/Source/Runtime/InteractiveToolsFramework/Private/BaseGizmos/GizmoElementGroup.cpp
@@ -139,7 +139,7 @@ FInputRayHit UGizmoElementGroupBase::LineTrace(const UGizmoViewContext* ViewCont
 					const auto ShouldUseNewHit = [&]() -> bool
 					{
 						// Negative priorities are ALWAYS behind positive priorities, even if exact hits.
-						if (bTestPriority && NewPriority < 0 && ExistingPriority > 0)
+						if (bTestPriority && NewPriority < 0 && ExistingPriority >= 0)
 						{
 							return false;
 						}
