Mark Input Behavior Capture Stealing functionality as UE_EXPERIMENTAL()

#rnx

diff --git a/Engine/Source/Runtime/InteractiveToolsFramework/Public/InputBehavior.h b/Engine/Source/Runtime/InteractiveToolsFramework/Public/InputBehavior.h
--- a/Engine/Source/Runtime/InteractiveToolsFramework/Public/InputBehavior.h
+++ b/Engine/Source/Runtime/InteractiveToolsFramework/Public/InputBehavior.h
@@ -310,17 +310,20 @@ public:
 	/**
 	 * When this function returns true, this behavior allows the Input Router to interrupt capturing and pass it on to a different Behavior with a higher priority.
 	 */
+	UE_EXPERIMENTAL(5.7, "Capture stealing is experimental behavior. Its implementation may change in the future.")
 	INTERACTIVETOOLSFRAMEWORK_API virtual bool AllowsCaptureStealing();
 
 	/**
 	 * Call to set the default capture stealing attribute.
 	 */
+	UE_EXPERIMENTAL(5.7, "Capture stealing is experimental behavior. Its implementation may change in the future.")
 	INTERACTIVETOOLSFRAMEWORK_API void SetAllowsCaptureStealing(TAttribute<bool> InAllowsCaptureStealing);
 
 protected:
 	/** priority returned by GetPriority() */
 	FInputCapturePriority DefaultPriority;
 
+	UE_EXPERIMENTAL(5.7, "Capture stealing is experimental behavior. Its implementation may change in the future.")
 	/** value returned by AllowsCaptureStealing*/
 	TAttribute<bool> CaptureStealing;
 };
