InteractiveToolsFramework: Guard CombinedTransformGizmo from the gizmo actor being destroyed out from under it.

diff --git a/Engine/Source/Runtime/InteractiveToolsFramework/Private/BaseGizmos/CombinedTransformGizmo.cpp b/Engine/Source/Runtime/InteractiveToolsFramework/Private/BaseGizmos/CombinedTransformGizmo.cpp
--- a/Engine/Source/Runtime/InteractiveToolsFramework/Private/BaseGizmos/CombinedTransformGizmo.cpp
+++ b/Engine/Source/Runtime/InteractiveToolsFramework/Private/BaseGizmos/CombinedTransformGizmo.cpp
@@ -820,7 +820,7 @@ bool UCombinedTransformGizmo::SetSubGizmoComponent(ETransformGizmoSubElements El
 	using namespace UE::Geometry;
 	using namespace CombinedTransformGizmoLocals;
 
-	if (!GizmoActor)
+	if (!IsValid(GizmoActor))
 	{
 		return false;
 	}
@@ -990,6 +990,11 @@ void UCombinedTransformGizmo::SetWorldAlignmentFunctions(
 	ShouldAlignDestination = MoveTemp(ShouldAlignTranslationIn);
 	DestinationAlignmentRayCaster = MoveTemp(TranslationAlignmentRayCasterIn);
 
+	if (!ensure(IsValid(GizmoActor)))
+	{
+		return;
+	}
+
 	// We allow this function to be called after Setup(), so modify any existing translation/rotation sub gizmos.
 	// Unfortunately we keep all the sub gizmos in one list, and the scaling gizmos are differentiated from the
 	// translation ones mainly in the components they use. So this ends up being a slightly messy set of checks,
@@ -1136,7 +1141,7 @@ void UCombinedTransformGizmo::Shutdown()
 {
 	ClearActiveTarget();
 
-	if (GizmoActor)
+	if (IsValid(GizmoActor))
 	{
 		GizmoActor->Destroy();
 		GizmoActor = nullptr;
@@ -1147,7 +1152,7 @@ void UCombinedTransformGizmo::Shutdown()
 
 void UCombinedTransformGizmo::UpdateCameraAxisSource()
 {
-	if (CameraAxisSource != nullptr && GizmoActor != nullptr)
+	if (CameraAxisSource != nullptr && IsValid(GizmoActor))
 	{
 		UE::GizmoUtil::UpdateCameraAxisSource(*CameraAxisSource, GetGizmoManager(), 
 			GizmoActor->GetTransform().GetLocation());
@@ -1327,6 +1332,11 @@ void UCombinedTransformGizmo::SetActiveTarget(UTransformProxy* Target, IToolCont
 		ClearActiveTarget();
 	}
 
+	if (!ensure(IsValid(GizmoActor)))
+	{
+		return;
+	}
+
 	ActiveTarget = Target;
 	TransactionProviderAtLastSetActiveTarget = TransactionProvider;
 
@@ -1438,12 +1448,21 @@ void UCombinedTransformGizmo::SetActiveTarget(UTransformProxy* Target, IToolCont
 
 FTransform UCombinedTransformGizmo::GetGizmoTransform() const
 {
+	if (!ensure(IsValid(GizmoActor)))
+	{
+		return FTransform::Identity;
+	}
 	USceneComponent* GizmoComponent = GizmoActor->GetRootComponent();
 	return GizmoComponent->GetComponentTransform();
 }
 
 void UCombinedTransformGizmo::ReinitializeGizmoTransform(const FTransform& NewTransform, bool bKeepGizmoUnscaled)
 {
+	if (!ensure(IsValid(GizmoActor)))
+	{
+		return;
+	}
+
 	// To update the gizmo location without triggering any callbacks, we temporarily
 	// store a copy of the callback list, detach them, reposition, and then reattach
 	// the callbacks.
@@ -1467,8 +1486,11 @@ void UCombinedTransformGizmo::ReinitializeGizmoTransform(const FTransform& NewTr
 
 void UCombinedTransformGizmo::SetNewGizmoTransform(const FTransform& NewTransform, bool bKeepGizmoUnscaled)
 {
-	check(ActiveTarget != nullptr);
-
+	if (!ensure(ActiveTarget))
+	{
+		return;
+	}
+	
 	BeginTransformEditSequence();
 	UpdateTransformDuringEditSequence(NewTransform, bKeepGizmoUnscaled);
 	EndTransformEditSequence();
@@ -1492,7 +1514,10 @@ void UCombinedTransformGizmo::EndTransformEditSequence()
 
 void UCombinedTransformGizmo::UpdateTransformDuringEditSequence(const FTransform& NewTransform, bool bKeepGizmoUnscaled)
 {
-	check(ActiveTarget != nullptr);
+	if (!ensure(ActiveTarget && IsValid(GizmoActor)))
+	{
+		return;
+	}
 
 	USceneComponent* GizmoComponent = GizmoActor->GetRootComponent();
 	FTransform GizmoTransform = NewTransform;
@@ -1518,6 +1543,11 @@ void UCombinedTransformGizmo::SetNewChildScale(const FVector& NewChildScale)
 
 void UCombinedTransformGizmo::SetVisibility(bool bVisible)
 {
+	if (!ensure(IsValid(GizmoActor)))
+	{
+		return;
+	}
+
 	bool bPreviousVisibility = !GizmoActor->IsHidden();
 
 	GizmoActor->SetActorHiddenInGame(bVisible == false);
@@ -1545,7 +1575,7 @@ ETransformGizmoSubElements UCombinedTransformGizmo::GetGizmoElements()
 {
 	using namespace CombinedTransformGizmoLocals;
 
-	return GetSubElementFlagsFromActor(GizmoActor);
+	return GetSubElementFlagsFromActor(GizmoActor.Get());
 }
 
 UInteractiveGizmo* UCombinedTransformGizmo::AddAxisTranslationGizmo(
diff --git a/Engine/Source/Runtime/InteractiveToolsFramework/Private/BaseGizmos/RepositionableTransformGizmo.cpp b/Engine/Source/Runtime/InteractiveToolsFramework/Private/BaseGizmos/RepositionableTransformGizmo.cpp
--- a/Engine/Source/Runtime/InteractiveToolsFramework/Private/BaseGizmos/RepositionableTransformGizmo.cpp
+++ b/Engine/Source/Runtime/InteractiveToolsFramework/Private/BaseGizmos/RepositionableTransformGizmo.cpp
@@ -53,6 +53,11 @@ UInteractiveGizmo* URepositionableTransformGizmoBuilder::BuildGizmo(const FToolB
 
 void URepositionableTransformGizmo::SetActiveTarget(UTransformProxy* Target, IToolContextTransactionProvider* TransactionProvider)
 {
+	if (!ensure(IsValid(GizmoActor)))
+	{
+		return;
+	}
+
 	// We're going to add on to the original transform gizmo, adding extra gizmos centered around
 	// existing gizmo components. These will respond to middle clicks and take "set pivot" code
 	// paths in their transform source and state target.
diff --git a/Engine/Source/Runtime/InteractiveToolsFramework/Public/BaseGizmos/CombinedTransformGizmo.h b/Engine/Source/Runtime/InteractiveToolsFramework/Public/BaseGizmos/CombinedTransformGizmo.h
--- a/Engine/Source/Runtime/InteractiveToolsFramework/Public/BaseGizmos/CombinedTransformGizmo.h
+++ b/Engine/Source/Runtime/InteractiveToolsFramework/Public/BaseGizmos/CombinedTransformGizmo.h
@@ -439,7 +439,7 @@ public:
 	/**
 	 * @return the internal GizmoActor used by the Gizmo
 	 */
-	ACombinedTransformGizmoActor* GetGizmoActor() const { return GizmoActor; }
+	ACombinedTransformGizmoActor* GetGizmoActor() const { return GizmoActor.Get(); }
 
 	/**
 	 * @return current transform of Gizmo
@@ -493,7 +493,7 @@ public:
 	 */
 	virtual bool IsVisible()
 	{
-		return GizmoActor && !GizmoActor->IsHidden();
+		return IsValid(GizmoActor) && !GizmoActor->IsHidden();
 	}
 
 	/**
@@ -664,8 +664,14 @@ protected:
 	/** GizmoActors will be spawned in this World */
 	UWorld* World;
 
-	/** Current active GizmoActor that was spawned by this Gizmo. Will be destroyed when Gizmo is. */
-	ACombinedTransformGizmoActor* GizmoActor;
+	/** 
+	 * Current active GizmoActor that was spawned by this Gizmo. Will be destroyed when
+	 *  this gizmo is, and expected not to be destroyed otherwise. Still, needs to be
+	 *  validated in case things go wrong and the world is destroyed before calling Shutdown
+	 *  on the gizmo.
+	 */
+	UPROPERTY()
+	TObjectPtr<ACombinedTransformGizmoActor> GizmoActor;
 
 	//
 	// Axis Sources
