[ITF & EditorGizmos] Code cleanup and experimental API removal

#DisabledValidations nonUnityAndStaticAnalysis

diff --git a/Engine/Source/Editor/Experimental/EditorInteractiveToolsFramework/Private/EditorGizmos/TransformGizmo.cpp b/Engine/Source/Editor/Experimental/EditorInteractiveToolsFramework/Private/EditorGizmos/TransformGizmo.cpp
--- a/Engine/Source/Editor/Experimental/EditorInteractiveToolsFramework/Private/EditorGizmos/TransformGizmo.cpp
+++ b/Engine/Source/Editor/Experimental/EditorInteractiveToolsFramework/Private/EditorGizmos/TransformGizmo.cpp
@@ -1287,12 +1287,12 @@ void UTransformGizmo::OnParametersChanged(const FGizmosParameters& InParameters)
 
 	// @note: The clamp values should match the property metadata (ClampMin)
 	Style.AxisSizeMultiplier = FMath::Max(0.01f, InParameters.AxisSizeMultiplier);
-	Style.LineThicknessMultiplier = FMath::Max(0.01f, InParameters.DevParameters.LineThicknessMultiplier);
-	Style.TranslateAxisLength = FMath::Max(1.0f, InParameters.DevParameters.TranslateAxisLength);
-	Style.RotateArcballSphereRadius = FMath::Max(1.0f, InParameters.DevParameters.RotateAxisRadius);
+	Style.LineThicknessMultiplier = 1.0f;
+	Style.TranslateAxisLength = 70.0f;
+	Style.RotateArcballSphereRadius = 70.0f;
 	Style.RotateAxisRadius = Style.RotateArcballSphereRadius;
 	Style.RotateCircleRadius =Style. RotateArcballSphereRadius;
-	Style.ScaleAxisLength = FMath::Max(1.0f, InParameters.DevParameters.ScaleAxisLength);
+	Style.ScaleAxisLength = 70.0f;
 
 	UpdateElements();
 }
@@ -2317,151 +2317,22 @@ void UTransformGizmo::OnClickPressAxis(const FInputDeviceRay& PressPos)
 
 void UTransformGizmo::SnapTranslateDelta(FVector& InOutWorldDelta) const
 {
-	UInteractiveGizmoManager* GizmoManager = GetGizmoManager();
-	if (!GizmoManager)
-	{
-		return;
-	}
-	
-	IToolsContextQueriesAPI* QueriesAPI = GizmoManager->GetContextQueriesAPI();
-	if (!QueriesAPI)
-	{
-		return;
-	}
-
-	if (const USceneSnappingManager* SnapManager = USceneSnappingManager::Find(GizmoManager))
-	{
-		FSceneSnapQueryRequest SnapRequest;
-		SnapRequest.RequestCoordinateSpace = QueriesAPI->GetCurrentCoordinateSystem();
-		SnapRequest.RequestType = ESceneSnapQueryType::Position;
-		SnapRequest.AxisList = InteractionAxisList;
-
-		FVector PositionToSnap = InOutWorldDelta;
-		if (SnapRequest.RequestCoordinateSpace == EToolContextCoordinateSystem::Local)
-		{
-			// Convert the delta to local space
-			const FQuat Rotation = ActiveTarget->GetTransform().GetRotation();
-			PositionToSnap = Rotation.Inverse() * InOutWorldDelta;
-		}
-
-		SnapRequest.Position = PositionToSnap;
-
-		TArray<FSceneSnapQueryResult> SnapResults;
-		if (SnapManager->ExecuteSceneSnapQuery(SnapRequest, SnapResults))
-		{
-			check(!SnapResults.IsEmpty());
-			const FSceneSnapQueryResult& SnapResult = SnapResults[0];
-
-			// Convert the delta back to the desired space
-			if (SnapResult.GetCoordinateSpace(SnapRequest) == EToolContextCoordinateSystem::Local)
-			{
-				const FQuat Rotation = ActiveTarget->GetTransform().GetRotation();
-				InOutWorldDelta = Rotation * SnapResult.Position;
-			}
-			else
-			{
-				InOutWorldDelta = SnapResult.Position;
-			}
-		}
-	}
+	// @note: not yet implemented, so doesn't modify the value
 }
 
 void UTransformGizmo::SnapRotateDelta(FQuat& InOutWorldDelta) const
 {
-	UInteractiveGizmoManager* GizmoManager = GetGizmoManager();
-	if (!GizmoManager)
-	{
-		return;
-	}
-
-	IToolsContextQueriesAPI* QueriesAPI = GizmoManager->GetContextQueriesAPI();
-	if (!QueriesAPI)
-	{
-		return;
-	}
-
-	if (const USceneSnappingManager* SnapManager = USceneSnappingManager::Find(GizmoManager))
-	{
-		FSceneSnapQueryRequest SnapRequest;
-		SnapRequest.RequestCoordinateSpace = QueriesAPI->GetCurrentCoordinateSystem();
-		SnapRequest.RequestType = ESceneSnapQueryType::Rotation;
-		SnapRequest.AxisList = InteractionAxisList;
-		SnapRequest.DeltaRotation = InOutWorldDelta;
-
-		TArray<FSceneSnapQueryResult> SnapResults;
-		if (SnapManager->ExecuteSceneSnapQuery(SnapRequest, SnapResults))
-		{
-			check(!SnapResults.IsEmpty());
-
-			InOutWorldDelta = SnapResults[0].DeltaRotation;
-		}
-	}
+	// @note: not yet implemented, so doesn't modify the value
 }
 
 void UTransformGizmo::SnapRotateAngleDelta(double& InOutAngleDelta) const
 {
-	UInteractiveGizmoManager* GizmoManager = GetGizmoManager();
-	if (!GizmoManager)
-	{
-		return;
-	}
-
-	IToolsContextQueriesAPI* QueriesAPI = GizmoManager->GetContextQueriesAPI();
-	if (!QueriesAPI)
-	{
-		return;
-	}
-
-	if (const USceneSnappingManager* SnapManager = USceneSnappingManager::Find(GizmoManager))
-	{
-		FSceneSnapQueryRequest SnapRequest;
-		SnapRequest.RequestCoordinateSpace = QueriesAPI->GetCurrentCoordinateSystem();
-		SnapRequest.RequestType = ESceneSnapQueryType::RotationAngle;
-		SnapRequest.AxisList = InteractionAxisList;
-		SnapRequest.RotationAngle = InOutAngleDelta;
-
-		TArray<FSceneSnapQueryResult> SnapResults;
-		if (SnapManager->ExecuteSceneSnapQuery(SnapRequest, SnapResults))
-		{
-			check(!SnapResults.IsEmpty());
-
-			InOutAngleDelta = SnapResults[0].RotationAngle;
-		}
-	}
+	// @note: not yet implemented, so doesn't modify the value
 }
 
 void UTransformGizmo::SnapScaleDelta(FVector& InOutLocalScaleDelta) const
 {
-	UInteractiveGizmoManager* GizmoManager = GetGizmoManager();
-	if (!GizmoManager)
-	{
-		return;
-	}
-
-	IToolsContextQueriesAPI* QueriesAPI = GizmoManager->GetContextQueriesAPI();
-	if (!QueriesAPI)
-	{
-		return;
-	}
-
-	if (const USceneSnappingManager* SnapManager = USceneSnappingManager::Find(GizmoManager))
-	{
-		FSceneSnapQueryRequest SnapRequest;
-
-		// Scale should always be in local space
-		SnapRequest.RequestCoordinateSpace = EToolContextCoordinateSystem::Local;
-		SnapRequest.RequestType = ESceneSnapQueryType::Scale;
-		SnapRequest.AxisList = InteractionAxisList;
-		SnapRequest.Scale = InOutLocalScaleDelta;
-
-		TArray<FSceneSnapQueryResult> SnapResults;
-		if (SnapManager->ExecuteSceneSnapQuery(SnapRequest, SnapResults))
-		{
-			check(!SnapResults.IsEmpty());
-
-			InOutLocalScaleDelta = SnapResults[0].Scale;
-		}
-	}
+	// @note: not yet implemented, so doesn't modify the value
 }
 
 void UTransformGizmo::OnClickDragTranslateAxis(const FInputDeviceRay& DragPos)
diff --git a/Engine/Source/Editor/Experimental/EditorInteractiveToolsFramework/Private/Snapping/EditorSnappingManager.cpp b/Engine/Source/Editor/Experimental/EditorInteractiveToolsFramework/Private/Snapping/EditorSnappingManager.cpp
--- a/Engine/Source/Editor/Experimental/EditorInteractiveToolsFramework/Private/Snapping/EditorSnappingManager.cpp
+++ b/Engine/Source/Editor/Experimental/EditorInteractiveToolsFramework/Private/Snapping/EditorSnappingManager.cpp
@@ -3,15 +3,10 @@
 #include "Snapping/EditorSnappingManager.h"
 
 #include "ContextObjectStore.h"
-#include "GridSnapping.h"
 #include "InteractiveToolsContext.h"
-#include "PolicySnapping.h"
 #include "SceneQueries/SceneSnappingManager.h"
 #include "ToolContextInterfaces.h"
 #include "VectorUtil.h"
-#include "ViewportSnappingModule.h"
-
-#include UE_INLINE_GENERATED_CPP_BY_NAME(EditorSnappingManager)
 
 namespace UE::Editor::Gizmos
 {
@@ -103,19 +98,10 @@ void UEditorSceneSnappingManager::Initialize(const TObjectPtr<UInteractiveToolsC
 {
 	UInteractiveToolManager* ToolManager = InToolsContext ? InToolsContext->ToolManager : nullptr;
 	QueriesAPI = ToolManager ? ToolManager->GetContextQueriesAPI() : nullptr;
-
-	RegisterSnapQueryTargetHandler(MakeUnique<UE::Editor::Gizmos::FGizmoGridSnapper>(InToolsContext));
-
-	if (const TSharedPtr<ISnappingPolicy>& MergedSnappingPolicy = IViewportSnappingModule::GetSnapManager())
-	{
-		RegisterSnapQueryTargetHandler(MakeUnique<UE::Editor::Gizmos::FGizmoPolicySnapper>(MergedSnappingPolicy));
-	}
 }
 
 void UEditorSceneSnappingManager::Shutdown()
 {
-	QueryTargetHandlers.Reset();
-
 	QueriesAPI = nullptr;
 }
 
@@ -127,165 +113,7 @@ bool UEditorSceneSnappingManager::ExecuteSceneHitQuery(const FSceneHitQueryReque
 
 bool UEditorSceneSnappingManager::ExecuteSceneSnapQuery(const FSceneSnapQueryRequest& InRequest, TArray<FSceneSnapQueryResult>& OutResults) const
 {
-	switch (InRequest.RequestType)
-	{
-	case ESceneSnapQueryType::Position:
-		return ExecuteSceneSnapQueryPosition(InRequest, OutResults);
-	case ESceneSnapQueryType::Rotation:
-		return ExecuteSceneSnapQueryRotation(InRequest, OutResults);
-	case ESceneSnapQueryType::RotationAngle:
-		return ExecuteSceneSnapQueryRotationAngle(InRequest, OutResults);
-	case ESceneSnapQueryType::Scale:
-		return ExecuteSceneSnapQueryScale(InRequest, OutResults);
-	default:
-		ensureMsgf(false, TEXT("Only Position, Rotation and Scale Snap Queries are supported"));
-	}
-
-	return false;
-}
-
-void UEditorSceneSnappingManager::RegisterSnapQueryTargetHandler(TUniquePtr<FSceneSnapQueryTargetHandler>&& InHandler)
-{
-	QueryTargetHandlers.Emplace(MoveTemp(InHandler));
-}
-
-void UEditorSceneSnappingManager::DeregisterSnapQueryTargetHandler(const FName InHandlerName)
-{
-	const int32 HandlerIdx = QueryTargetHandlers.IndexOfByPredicate(
-			[InHandlerName](const TUniquePtr<FSceneSnapQueryTargetHandler>& InHandler)
-			{
-				return InHandler->TargetName == InHandlerName;
-			});
-	
-	if (HandlerIdx != INDEX_NONE)
-	{
-		QueryTargetHandlers.RemoveAt(HandlerIdx);
-	}
-}
-
-bool UEditorSceneSnappingManager::ExecuteSceneSnapQueryPosition(const FSceneSnapQueryRequest& InRequest, TArray<FSceneSnapQueryResult>& OutResults) const
-{
-	using namespace UE::Editor::Gizmos;
-
-	if (!QueriesAPI)
-	{
-		return false;
-	}
-
-	for (const TUniquePtr<FSceneSnapQueryTargetHandler>& Handler : QueryTargetHandlers)
-	{
-		if (Handler->IsQueryTypeSupported(ESceneSnapQueryType::Position))
-		{
-			FVector SnappedPosition;
-			if (Handler->SnapPosition(InRequest.Position, SnappedPosition) == ESceneSnapQueryTargetResult::Snapped)
-			{
-				FSceneSnapQueryResult SnapResult;
-				SnapResult.TargetType = Handler->GetTargetType();
-				SnapResult.Position = SnappedPosition;
-
-				OutResults.Emplace(SnapResult);
-
-				return true;
-			}
-		}
-	}
-
-	return false;
-}
-
-bool UEditorSceneSnappingManager::ExecuteSceneSnapQueryRotation(const FSceneSnapQueryRequest& InRequest, TArray<FSceneSnapQueryResult>& OutResults) const
-{
-	using namespace UE::Editor::Gizmos;
-
-	if (!QueriesAPI)
-	{
-		return false;
-	}
-
-	for (const TUniquePtr<FSceneSnapQueryTargetHandler>& Handler : QueryTargetHandlers)
-	{
-		if (Handler->IsQueryTypeSupported(ESceneSnapQueryType::Rotation))
-		{
-			FQuat SnappedRotation;
-			if (Handler->SnapRotation(InRequest.DeltaRotation, SnappedRotation) == ESceneSnapQueryTargetResult::Snapped)
-			{
-				FSceneSnapQueryResult SnapResult;
-				SnapResult.TargetType = Handler->GetTargetType();
-				SnapResult.DeltaRotation = SnappedRotation;
-
-				OutResults.Emplace(SnapResult);
-
-				return true;
-			}
-		}
-	}
-
-	return false;
-}
-
-bool UEditorSceneSnappingManager::ExecuteSceneSnapQueryRotationAngle(const FSceneSnapQueryRequest& InRequest, TArray<FSceneSnapQueryResult>& OutResults) const
-{
-	using namespace UE::Editor::Gizmos;
-
-	if (!QueriesAPI)
-	{
-		return false;
-	}
-	
-	// It's assumed the input AxisList only has a single active axis
-	const EAxis::Type Axis = EditorSnappingManagerLocals::GetAxisFromAxisList(InRequest.AxisList);
-	if (Axis == EAxis::None)
-	{
-		return false;
-	}
-
-	for (const TUniquePtr<FSceneSnapQueryTargetHandler>& Handler : QueryTargetHandlers)
-	{
-		if (Handler->IsQueryTypeSupported(ESceneSnapQueryType::RotationAngle))
-		{
-			double SnappedAngle;
-			if (Handler->SnapRotationAxisAngle(InRequest.RotationAngle, SnappedAngle, Axis) == ESceneSnapQueryTargetResult::Snapped)
-			{
-				FSceneSnapQueryResult SnapResult;
-				SnapResult.TargetType = Handler->GetTargetType();
-				SnapResult.RotationAngle = SnappedAngle;
-
-				OutResults.Emplace(SnapResult);
-
-				return true;
-			}
-		}
-	}
-
-	return false;
-}
-
-bool UEditorSceneSnappingManager::ExecuteSceneSnapQueryScale(const FSceneSnapQueryRequest& InRequest, TArray<FSceneSnapQueryResult>& OutResults) const
-{
-	using namespace UE::Editor::Gizmos;
-
-	if (!QueriesAPI)
-	{
-		return false;
-	}
-
-	for (const TUniquePtr<FSceneSnapQueryTargetHandler>& Handler : QueryTargetHandlers)
-	{
-		if (Handler->IsQueryTypeSupported(ESceneSnapQueryType::Scale))
-		{
-			FVector SnappedScale;
-			if (Handler->SnapScale(InRequest.Scale, SnappedScale) == ESceneSnapQueryTargetResult::Snapped)
-			{
-				FSceneSnapQueryResult SnapResult;
-				SnapResult.TargetType = Handler->GetTargetType();
-				SnapResult.Scale = SnappedScale;
-
-				OutResults.Emplace(SnapResult);
-
-				return true;
-			}
-		}
-	}
+	unimplemented();
 
 	return false;
 }
diff --git a/Engine/Source/Editor/Experimental/EditorInteractiveToolsFramework/Private/Snapping/GridSnapping.cpp b/Engine/Source/Editor/Experimental/EditorInteractiveToolsFramework/Private/Snapping/GridSnapping.cpp
deleted file mode 100644
diff --git a/Engine/Source/Editor/Experimental/EditorInteractiveToolsFramework/Private/Snapping/GridSnapping.h b/Engine/Source/Editor/Experimental/EditorInteractiveToolsFramework/Private/Snapping/GridSnapping.h
deleted file mode 100644
diff --git a/Engine/Source/Editor/Experimental/EditorInteractiveToolsFramework/Private/Snapping/PolicySnapping.cpp b/Engine/Source/Editor/Experimental/EditorInteractiveToolsFramework/Private/Snapping/PolicySnapping.cpp
deleted file mode 100644
diff --git a/Engine/Source/Editor/Experimental/EditorInteractiveToolsFramework/Private/Snapping/PolicySnapping.h b/Engine/Source/Editor/Experimental/EditorInteractiveToolsFramework/Private/Snapping/PolicySnapping.h
deleted file mode 100644
diff --git a/Engine/Source/Editor/Experimental/EditorInteractiveToolsFramework/Public/EditorGizmos/TransformGizmo.h b/Engine/Source/Editor/Experimental/EditorInteractiveToolsFramework/Public/EditorGizmos/TransformGizmo.h
--- a/Engine/Source/Editor/Experimental/EditorInteractiveToolsFramework/Public/EditorGizmos/TransformGizmo.h
+++ b/Engine/Source/Editor/Experimental/EditorInteractiveToolsFramework/Public/EditorGizmos/TransformGizmo.h
@@ -57,19 +57,6 @@ class UObject;
 class UTransformProxy;
 class UGizmoElementGimbal;
 
-// @NOTE: The following changes should be made prior to 5.7 Release:
-// De-virtualize:
-//  HandleWidgetModeChanged
-//  OnParametersChanged
-//  SetupIndirectBehaviors
-//  UpdateRotationMode
-//  GetGimbalRotationAxis
-//  ResetHoverStates
-//  ResetInteractingStates
-//  OnClickPressRotateAxis
-//  ComputeGimbalRotateDelta
-//  OnClickPressGimbalRotateAxis
-
 //
 // Interface for the Transform gizmo.
 //
@@ -124,25 +111,6 @@ namespace EAxisRotateMode
 	};
 }
 
-/** @todo: To be removed for release, for dev/tweaking only. */
-USTRUCT()
-struct FGizmoDevParameters
-{
-	GENERATED_BODY()
-
-	UPROPERTY(EditAnywhere, Category = "Transform Gizmo", meta = (ClampMin = 0.01f))
-	float LineThicknessMultiplier = 1.0f;
-
-	UPROPERTY(EditAnywhere, Category = "Transform Gizmo", meta = (ClampMin = 1.0f))
-	float TranslateAxisLength = 70.0f;
-
-	UPROPERTY(EditAnywhere, Category = "Transform Gizmo", meta = (ClampMin = 1.0f))
-	float ScaleAxisLength = 70.0f;
-
-	UPROPERTY(EditAnywhere, Category = "Transform Gizmo", meta = (ClampMin = 1.0f))
-	float RotateAxisRadius = 70.0f;
-};
-
 USTRUCT()
 struct FGizmosParameters
 {
@@ -169,9 +137,6 @@ struct FGizmosParameters
 	/** Multiplies the values of TranslateAxisLength, ScaleAxisLength, and RotateAxisRadius. */
 	UPROPERTY()
 	float AxisSizeMultiplier = 1.0f;
-
-	UPROPERTY(EditAnywhere, Category = "Transform Gizmo|Development", meta=(ShowOnlyInnerProperties))
-	FGizmoDevParameters DevParameters;
 };
 
 /**
diff --git a/Engine/Source/Editor/Experimental/EditorInteractiveToolsFramework/Public/Snapping/EditorSnappingManager.h b/Engine/Source/Editor/Experimental/EditorInteractiveToolsFramework/Public/Snapping/EditorSnappingManager.h
--- a/Engine/Source/Editor/Experimental/EditorInteractiveToolsFramework/Public/Snapping/EditorSnappingManager.h
+++ b/Engine/Source/Editor/Experimental/EditorInteractiveToolsFramework/Public/Snapping/EditorSnappingManager.h
@@ -39,17 +39,6 @@ public:
 	virtual bool ExecuteSceneHitQuery(const FSceneHitQueryRequest& InRequest, FSceneHitQueryResult& OutResult) const override;
 	virtual bool ExecuteSceneSnapQuery(const FSceneSnapQueryRequest& InRequest, TArray<FSceneSnapQueryResult>& OutResults) const override;
 
-	/** Register a snapping handler for a given snap target. */
-	EDITORINTERACTIVETOOLSFRAMEWORK_API void RegisterSnapQueryTargetHandler(TUniquePtr<FSceneSnapQueryTargetHandler>&& InHandler);
-	EDITORINTERACTIVETOOLSFRAMEWORK_API void DeregisterSnapQueryTargetHandler(const FName InHandlerName);
-
 protected:
-	bool ExecuteSceneSnapQueryPosition(const FSceneSnapQueryRequest& InRequest, TArray<FSceneSnapQueryResult>& OutResults) const;
-	bool ExecuteSceneSnapQueryRotation(const FSceneSnapQueryRequest& InRequest, TArray<FSceneSnapQueryResult>& OutResults) const;
-	bool ExecuteSceneSnapQueryRotationAngle(const FSceneSnapQueryRequest& InRequest, TArray<FSceneSnapQueryResult>& OutResults) const;
-	bool ExecuteSceneSnapQueryScale(const FSceneSnapQueryRequest& InRequest, TArray<FSceneSnapQueryResult>& OutResults) const;
-
-	TArray<TUniquePtr<FSceneSnapQueryTargetHandler>> QueryTargetHandlers;
-
 	const IToolsContextQueriesAPI* QueriesAPI = nullptr;
 };
diff --git a/Engine/Source/Runtime/InteractiveToolsFramework/Public/SceneQueries/SceneSnappingManager.h b/Engine/Source/Runtime/InteractiveToolsFramework/Public/SceneQueries/SceneSnappingManager.h
--- a/Engine/Source/Runtime/InteractiveToolsFramework/Public/SceneQueries/SceneSnappingManager.h
+++ b/Engine/Source/Runtime/InteractiveToolsFramework/Public/SceneQueries/SceneSnappingManager.h
@@ -129,81 +129,6 @@ enum class ESceneSnapQueryTargetResult : uint8
 	Unsupported
 };
 
-/** A custom query target implementation, allowing extensibility and handling
- * of custom targets not defined in the ESceneSnapQueryTargetType enum. */
-class FSceneSnapQueryTargetHandler
-{
-public:
-	virtual ~FSceneSnapQueryTargetHandler() = default;
-
-	/** Snaps the given transform, potentially affecting all query types (position, rotation, scale). The given transform can be absolute or a relative delta. */
-	virtual ESceneSnapQueryTargetResult SnapTransform(const FTransform& InTransform, FVector& OutSnappedTransform) const { return ESceneSnapQueryTargetResult::Unsupported; }
-
-	virtual ESceneSnapQueryTargetResult SnapPosition(const FVector& InPosition, FVector& OutSnappedPosition, const EAxisList::Type InAxisList = EAxisList::All) const { return ESceneSnapQueryTargetResult::Unsupported; }
-	virtual ESceneSnapQueryTargetResult SnapPositionAxis(const double InAxisValue, double& OutSnappedAxisValue, const EAxis::Type InAxis) const { return ESceneSnapQueryTargetResult::Unsupported; }
-
-	virtual ESceneSnapQueryTargetResult SnapRotation(const FQuat& InRotation, FQuat& OutSnappedRotation, const EAxisList::Type InAxisList = EAxisList::All) const { return ESceneSnapQueryTargetResult::Unsupported; }
-	virtual ESceneSnapQueryTargetResult SnapRotationAxisAngle(const double& InAngle, double& OutSnappedAngle, const EAxis::Type InAxis) const { return ESceneSnapQueryTargetResult::Unsupported; }
-
-	virtual ESceneSnapQueryTargetResult SnapScale(const FVector& InScale, FVector& OutSnappedScale, const EAxisList::Type InAxisList = EAxisList::All) const { return ESceneSnapQueryTargetResult::Unsupported; }
-	virtual ESceneSnapQueryTargetResult SnapScaleAxis(const double InAxisValue, double& OutSnappedAxisValue, const EAxis::Type InAxis) const { return ESceneSnapQueryTargetResult::Unsupported; }
-
-	virtual ESceneSnapQueryTargetType GetTargetType() const { return ESceneSnapQueryTargetType::Custom; }
-
-	/** Returns true if the given query type can -ever- be supported, regardless of any variable state. */
-	virtual bool IsQueryTypeSupported(const ESceneSnapQueryType InQueryType) const { return false; }
-
-	bool operator==(const FSceneSnapQueryTargetHandler& InOther) const
-	{
-		return TargetName == InOther.TargetName;
-	}
-
-	bool operator !=(const FSceneSnapQueryTargetHandler& InOther) const
-	{
-		return TargetName != InOther.TargetName;
-	}
-
-public:
-	/** Unique query target name, ie. "Grid" */
-	FName TargetName;
-};
-
-inline uint32 GetTypeHash(const FSceneSnapQueryTargetHandler& InSnapQueryTargetHandler)
-{
-	return GetTypeHash(InSnapQueryTargetHandler.TargetName);
-}
-
-/** A custom query target implementation with typed parameters, allowing extensibility and handling
- * of custom targets not defined in the ESceneSnapQueryTargetType enum. */
-template <typename TargetRequestParametersType UE_REQUIRES(std::is_class_v<TargetRequestParametersType>)>
-class TSceneSnapQueryTargetHandler : public FSceneSnapQueryTargetHandler
-{
-public:
-	/** Snaps the given transform, potentially affecting all query types (position, rotation, scale). The given transform can be absolute or a relative delta. */
-	virtual ESceneSnapQueryTargetResult SnapTransform(const FTransform& InTransform, FVector& OutSnappedTransform) const override { return ESceneSnapQueryTargetResult::Unsupported; }
-
-	/** Snaps the given transform, potentially affecting all query types (position, rotation, scale). The given transform can be absolute or a relative delta. */
-	virtual ESceneSnapQueryTargetResult SnapTransform(const TargetRequestParametersType& InTargetRequest, const FTransform& InTransform, FVector& OutSnappedTransform) const { return ESceneSnapQueryTargetResult::Unsupported; }
-
-	virtual ESceneSnapQueryTargetResult SnapPosition(const FVector& InPosition, FVector& OutSnappedPosition, const EAxisList::Type InAxisList = EAxisList::All) const override { return ESceneSnapQueryTargetResult::Unsupported; }
-	virtual ESceneSnapQueryTargetResult SnapPosition(const TargetRequestParametersType& InTargetRequest, const FVector& InPosition, FVector& OutSnappedPosition, const EAxisList::Type InAxisList = EAxisList::All) const { return ESceneSnapQueryTargetResult::Unsupported; }
-
-	virtual ESceneSnapQueryTargetResult SnapPositionAxis(const double InAxisValue, double& OutSnappedAxisValue, const EAxis::Type InAxis) const override { return ESceneSnapQueryTargetResult::Unsupported; }
-	virtual ESceneSnapQueryTargetResult SnapPositionAxis(const TargetRequestParametersType& InTargetRequest, const double InAxisValue, double& OutSnappedAxisValue, const EAxis::Type InAxis) const { return ESceneSnapQueryTargetResult::Unsupported; }
-
-	virtual ESceneSnapQueryTargetResult SnapRotation(const FQuat& InRotation, FQuat& OutSnappedRotation, const EAxisList::Type InAxisList = EAxisList::All) const override { return ESceneSnapQueryTargetResult::Unsupported; }
-	virtual ESceneSnapQueryTargetResult SnapRotation(const TargetRequestParametersType& InTargetRequest, const FQuat& InRotation, FQuat& OutSnappedRotation, const EAxisList::Type InAxisList = EAxisList::All) const { return ESceneSnapQueryTargetResult::Unsupported; }
-
-	virtual ESceneSnapQueryTargetResult SnapRotationAxisAngle(const double& InAngle, double& OutSnappedAngle, const EAxis::Type InAxis) const override { return ESceneSnapQueryTargetResult::Unsupported; }
-	virtual ESceneSnapQueryTargetResult SnapRotationAxisAngle(const TargetRequestParametersType& InTargetRequest, const double& InAngle, double& OutSnappedAngle, const EAxis::Type InAxis) const { return ESceneSnapQueryTargetResult::Unsupported; }
-
-	virtual ESceneSnapQueryTargetResult SnapScale(const FVector& InScale, FVector& OutSnappedScale, const EAxisList::Type InAxisList = EAxisList::All) const override { return ESceneSnapQueryTargetResult::Unsupported; }
-	virtual ESceneSnapQueryTargetResult SnapScale(const TargetRequestParametersType& InTargetRequest, const FVector& InScale, FVector& OutSnappedScale, const EAxisList::Type InAxisList = EAxisList::All) const { return ESceneSnapQueryTargetResult::Unsupported; }
-
-	virtual ESceneSnapQueryTargetResult SnapScaleAxis(const double InAxisValue, double& OutSnappedAxisValue, const EAxis::Type InAxis) const override { return ESceneSnapQueryTargetResult::Unsupported; }
-	virtual ESceneSnapQueryTargetResult SnapScaleAxis(const TargetRequestParametersType& InTargetRequest, const double InAxisValue, double& OutSnappedAxisValue, const EAxis::Type InAxis) const { return ESceneSnapQueryTargetResult::Unsupported; }
-};
-
 /**
  * Configuration variables for a USceneSnappingManager snap query request.
  */
