InteractiveToolsFramework: Enable snapping on the free translate gizmo.

diff --git a/Engine/Source/Runtime/InteractiveToolsFramework/Private/BaseGizmos/CombinedTransformGizmo.cpp b/Engine/Source/Runtime/InteractiveToolsFramework/Private/BaseGizmos/CombinedTransformGizmo.cpp
--- a/Engine/Source/Runtime/InteractiveToolsFramework/Private/BaseGizmos/CombinedTransformGizmo.cpp
+++ b/Engine/Source/Runtime/InteractiveToolsFramework/Private/BaseGizmos/CombinedTransformGizmo.cpp
@@ -1758,6 +1758,18 @@ UInteractiveGizmo* UCombinedTransformGizmo::AddFreeTranslationGizmo(
 		return nullptr;
 	}
 	ensure(Gizmo->InitializeAsScreenPlaneTranslateGizmo(Params, &SharedState));
+	if (UGizmoPlaneTranslationParameterSource* ParameterSource = Cast<UGizmoPlaneTranslationParameterSource>(Gizmo->ParameterSource.GetObject()))
+	{
+		ParameterSource->PositionConstraintFunction = [this](const FVector& Pos, FVector& Snapped) { return PositionSnapFunction(Pos, Snapped); };
+		ParameterSource->bProjectConstrainedPosition = false;
+		ParameterSource->AxisXDeltaConstraintFunction = [this, XAxisIndex = 0](double ScaleAxisDelta, double& SnappedScaleAxisDelta) { return PositionAxisDeltaSnapFunction(ScaleAxisDelta, SnappedScaleAxisDelta, XAxisIndex); };
+		ParameterSource->AxisYDeltaConstraintFunction = [this, YAxisIndex = 1](double ScaleAxisDelta, double& SnappedScaleAxisDelta) { return PositionAxisDeltaSnapFunction(ScaleAxisDelta, SnappedScaleAxisDelta, YAxisIndex); };
+	}
+	else
+	{
+		ensure(false);
+	}
+
 
 	ActiveComponents.Add(Params.Component);
 	ActiveGizmos.Add(Gizmo);
diff --git a/Engine/Source/Runtime/InteractiveToolsFramework/Private/BaseGizmos/ParameterToTransformAdapters.cpp b/Engine/Source/Runtime/InteractiveToolsFramework/Private/BaseGizmos/ParameterToTransformAdapters.cpp
--- a/Engine/Source/Runtime/InteractiveToolsFramework/Private/BaseGizmos/ParameterToTransformAdapters.cpp
+++ b/Engine/Source/Runtime/InteractiveToolsFramework/Private/BaseGizmos/ParameterToTransformAdapters.cpp
@@ -92,8 +92,11 @@ void UGizmoPlaneTranslationParameterSource::SetParameter(const FVector2D& NewVal
 	FVector SnappedPos;
 	if (PositionConstraintFunction(NewTransform.GetTranslation(), SnappedPos))
 	{
-		FVector PlanePos = GizmoMath::ProjectPointOntoPlane(SnappedPos, CurTranslationOrigin, CurTranslationNormal);
-		NewTransform.SetTranslation(PlanePos);
+		if (bProjectConstrainedPosition)
+		{
+			SnappedPos = GizmoMath::ProjectPointOntoPlane(SnappedPos, CurTranslationOrigin, CurTranslationNormal);
+		}
+		NewTransform.SetTranslation(SnappedPos);
 	}
 
 	TransformSource->SetTransform(NewTransform);
diff --git a/Engine/Source/Runtime/InteractiveToolsFramework/Public/BaseGizmos/ParameterToTransformAdapters.h b/Engine/Source/Runtime/InteractiveToolsFramework/Public/BaseGizmos/ParameterToTransformAdapters.h
--- a/Engine/Source/Runtime/InteractiveToolsFramework/Public/BaseGizmos/ParameterToTransformAdapters.h
+++ b/Engine/Source/Runtime/InteractiveToolsFramework/Public/BaseGizmos/ParameterToTransformAdapters.h
@@ -133,11 +133,13 @@ public:
 	/**
 	 * Optional position constraint function. Called during interaction with the new transform origin.
 	 * To snap the transform to a new position, return as second value, and return true from your lambda.
-	 * Note that returned snap point will be projected onto the current translation origin/axis.
 	 * @return true if constraint point was found and should be used, false to ignore
 	 */
 	TUniqueFunction<bool(const FVector& Input, FVector& PreProjectedOutput)> PositionConstraintFunction = [](const FVector&, FVector&) { return false; };
-
+	/*
+	 * When true, the output of PositionConstraintFunction will be projected onto the current translation origin/axis.
+	 */
+	bool bProjectConstrainedPosition = true;
 
 	virtual FVector2D GetParameter() const override
 	{
