[TextToAnim] Added the initial plugin code for text to animation.

#rnx
#virtualized

diff --git a/Engine/Source/Runtime/InteractiveToolsFramework/Private/InteractiveToolsContext.cpp b/Engine/Source/Runtime/InteractiveToolsFramework/Private/InteractiveToolsContext.cpp
--- a/Engine/Source/Runtime/InteractiveToolsFramework/Private/InteractiveToolsContext.cpp
+++ b/Engine/Source/Runtime/InteractiveToolsFramework/Private/InteractiveToolsContext.cpp
@@ -161,7 +161,7 @@ void UInteractiveToolsContext::DeactivateActiveTool(EToolSide WhichSide, EToolSh
 
 void UInteractiveToolsContext::DeactivateAllActiveTools(EToolShutdownType ShutdownType)
 {
-	if (!ensureMsgf(ToolManager, TEXT("Trying to deactivate all tools on a ToolManager that has already been deleted.")))
+	if (!ToolManager)
 	{
 		return;
 	}
