[ITF] Addressed some minor bughawk complaints from CL 45268317

diff --git a/Engine/Source/Editor/Experimental/EditorInteractiveToolsFramework/Private/EditorGizmos/GizmoElementGimbal.cpp b/Engine/Source/Editor/Experimental/EditorInteractiveToolsFramework/Private/EditorGizmos/GizmoElementGimbal.cpp
--- a/Engine/Source/Editor/Experimental/EditorInteractiveToolsFramework/Private/EditorGizmos/GizmoElementGimbal.cpp
+++ b/Engine/Source/Editor/Experimental/EditorInteractiveToolsFramework/Private/EditorGizmos/GizmoElementGimbal.cpp
@@ -70,6 +70,7 @@ FInputRayHit UGizmoElementGimbal::LineTrace(const UGizmoViewContext* ViewContext
 			FInputRayHit NewHit = Element->LineTrace(ViewContext, CurrentLineTraceState, RayOrigin, RayDirection, NewLineTraceOutput);
 			if (!Hit.bHit || NewHit.HitDepth < Hit.HitDepth)
 			{
+				OutLineTraceOutput = NewLineTraceOutput;
 				Hit = NewHit;
 			}
 		}
diff --git a/Engine/Source/Runtime/InteractiveToolsFramework/Private/BaseGizmos/GizmoElementBox.cpp b/Engine/Source/Runtime/InteractiveToolsFramework/Private/BaseGizmos/GizmoElementBox.cpp
--- a/Engine/Source/Runtime/InteractiveToolsFramework/Private/BaseGizmos/GizmoElementBox.cpp
+++ b/Engine/Source/Runtime/InteractiveToolsFramework/Private/BaseGizmos/GizmoElementBox.cpp
@@ -52,7 +52,7 @@ FInputRayHit UGizmoElementBox::LineTrace(const UGizmoViewContext* ViewContext, c
 
 			// We successfully hit, now check if we intersect the actual surface (rather than just within the PixelHitDistanceThreshold) ...
 			PixelHitThresholdAdjust = CurrentLineTraceState.PixelToWorldScale * MinimumPixelHitDistanceThreshold;
-			Box = UE::Geometry::TOrientedBox3<double>(Frame, WorldExtent + PixelHitThresholdAdjust); // Recreate the box without PixelHitDistanceThreshold
+			Box = UE::Geometry::TOrientedBox3<double>(Frame, WorldExtent + PixelHitThresholdAdjust); // Recreate the box with MinimumPixelHitDistanceThreshold
 			if (UE::Geometry::TIntrRay3OrientedBox3<double>::FindIntersection(Ray, Box, HitDepth))
 			{
 				RayHit.HitDepth = HitDepth;
diff --git a/Engine/Source/Runtime/InteractiveToolsFramework/Private/BaseGizmos/GizmoElementCone.cpp b/Engine/Source/Runtime/InteractiveToolsFramework/Private/BaseGizmos/GizmoElementCone.cpp
--- a/Engine/Source/Runtime/InteractiveToolsFramework/Private/BaseGizmos/GizmoElementCone.cpp
+++ b/Engine/Source/Runtime/InteractiveToolsFramework/Private/BaseGizmos/GizmoElementCone.cpp
@@ -147,7 +147,7 @@ FInputRayHit UGizmoElementCone::LineTrace(const UGizmoViewContext* ViewContext,
 
 			if (bIntersects)
 			{
-				RayHit.HitDepth = RayParam;
+				RayHit.HitDepth = RayParam + DepthBias;
 				OutLineTraceOutput.bIsSurfaceHit = true;
 			}
 
diff --git a/Engine/Source/Runtime/InteractiveToolsFramework/Public/BaseGizmos/GizmoElementBox.h b/Engine/Source/Runtime/InteractiveToolsFramework/Public/BaseGizmos/GizmoElementBox.h
--- a/Engine/Source/Runtime/InteractiveToolsFramework/Public/BaseGizmos/GizmoElementBox.h
+++ b/Engine/Source/Runtime/InteractiveToolsFramework/Public/BaseGizmos/GizmoElementBox.h
@@ -10,7 +10,7 @@
 
 /**
  * Simple object intended to be used as part of 3D Gizmos.
- * Draws a solid 3D cylinder based on parameters.
+ * Draws a solid 3D box based on parameters.
  */
 
 UCLASS(Transient, MinimalAPI)
@@ -25,7 +25,7 @@ public:
 	INTERACTIVETOOLSFRAMEWORK_API virtual FInputRayHit LineTrace(const UGizmoViewContext* ViewContext, const FLineTraceTraversalState& LineTraceState, const FVector& RayOrigin, const FVector& RayDirection, FLineTraceOutput& OutLineTraceOutput) override;
 	//~ End UGizmoElementBase Interface.
 
-	// Location of center of box's base circle.
+	// Location of center of box's base.
 	INTERACTIVETOOLSFRAMEWORK_API virtual void SetCenter(const FVector& InCenter);
 	INTERACTIVETOOLSFRAMEWORK_API virtual FVector GetCenter() const;
 
