[ITF] Added some const getter overloads

diff --git a/Engine/Source/Runtime/InteractiveToolsFramework/Private/BaseGizmos/GizmoElementRenderState.cpp b/Engine/Source/Runtime/InteractiveToolsFramework/Private/BaseGizmos/GizmoElementRenderState.cpp
--- a/Engine/Source/Runtime/InteractiveToolsFramework/Private/BaseGizmos/GizmoElementRenderState.cpp
+++ b/Engine/Source/Runtime/InteractiveToolsFramework/Private/BaseGizmos/GizmoElementRenderState.cpp
@@ -159,6 +159,25 @@ void FGizmoElementMeshRenderStateAttributes::Update(FGizmoElementMeshRenderState
 }
 
 FLinearColor FGizmoElementLineRenderStateAttributes::GetLineColor(EGizmoElementInteractionState InteractionState)
+{
+	return GetLineColorInternal(InteractionState);
+}
+
+FLinearColor FGizmoElementLineRenderStateAttributes::GetLineColor(EGizmoElementInteractionState InteractionState) const
+{
+	return GetLineColorInternal(InteractionState);
+}
+
+void FGizmoElementLineRenderStateAttributes::Update(FGizmoElementLineRenderStateAttributes& InChildAttributes)
+{
+	LineColor.UpdateState(InChildAttributes.LineColor);
+	HoverLineColor.UpdateState(InChildAttributes.HoverLineColor);
+	InteractLineColor.UpdateState(InChildAttributes.InteractLineColor);
+	SelectLineColor.UpdateState(InChildAttributes.SelectLineColor);
+	SubdueLineColor.UpdateState(InChildAttributes.SubdueLineColor);
+}
+
+FLinearColor FGizmoElementLineRenderStateAttributes::GetLineColorInternal(EGizmoElementInteractionState InteractionState) const
 {
 	switch (InteractionState)
 	{
@@ -180,12 +199,3 @@ FLinearColor FGizmoElementLineRenderStateAttributes::GetLineColor(EGizmoElementI
 
 	return LineColor.GetColor();
 }
-
-void FGizmoElementLineRenderStateAttributes::Update(FGizmoElementLineRenderStateAttributes& InChildAttributes)
-{
-	LineColor.UpdateState(InChildAttributes.LineColor);
-	HoverLineColor.UpdateState(InChildAttributes.HoverLineColor);
-	InteractLineColor.UpdateState(InChildAttributes.InteractLineColor);
-	SelectLineColor.UpdateState(InChildAttributes.SelectLineColor);
-	SubdueLineColor.UpdateState(InChildAttributes.SubdueLineColor);
-}
diff --git a/Engine/Source/Runtime/InteractiveToolsFramework/Public/BaseGizmos/GizmoElementBase.h b/Engine/Source/Runtime/InteractiveToolsFramework/Public/BaseGizmos/GizmoElementBase.h
--- a/Engine/Source/Runtime/InteractiveToolsFramework/Public/BaseGizmos/GizmoElementBase.h
+++ b/Engine/Source/Runtime/InteractiveToolsFramework/Public/BaseGizmos/GizmoElementBase.h
@@ -79,17 +79,35 @@ public:
 			return MeshRenderState.GetMaterial(InteractionState);
 		}
 
+		// Returns the mesh material based on the current interaction state.
+		const UMaterialInterface* GetCurrentMaterial() const
+		{
+			return MeshRenderState.GetMaterial(InteractionState);
+		}
+
 		// Returns the mesh vertex color based on the current interaction state.
 		FLinearColor GetCurrentVertexColor()
 		{
 			return MeshRenderState.GetVertexColor(InteractionState);
 		}
 
+		// Returns the mesh vertex color based on the current interaction state.
+		FLinearColor GetCurrentVertexColor() const
+		{
+			return MeshRenderState.GetVertexColor(InteractionState);
+		}
+
 		// Returns the line color based on the current interaction state.
 		FLinearColor GetCurrentLineColor()
 		{
 			return LineRenderState.GetLineColor(InteractionState);
 		}
+
+		// Returns the line color based on the current interaction state.
+		FLinearColor GetCurrentLineColor() const
+		{
+			return LineRenderState.GetLineColor(InteractionState);
+		}
 	};
 
 	struct FLineTraceTraversalState
diff --git a/Engine/Source/Runtime/InteractiveToolsFramework/Public/BaseGizmos/GizmoElementRenderState.h b/Engine/Source/Runtime/InteractiveToolsFramework/Public/BaseGizmos/GizmoElementRenderState.h
--- a/Engine/Source/Runtime/InteractiveToolsFramework/Public/BaseGizmos/GizmoElementRenderState.h
+++ b/Engine/Source/Runtime/InteractiveToolsFramework/Public/BaseGizmos/GizmoElementRenderState.h
@@ -187,7 +187,14 @@ struct FGizmoElementLineRenderStateAttributes
 	// Returns the line color corresponding to the input interaction state
 	INTERACTIVETOOLSFRAMEWORK_API FLinearColor GetLineColor(EGizmoElementInteractionState InteractionState);
 
+	// Returns the line color corresponding to the input interaction state
+	INTERACTIVETOOLSFRAMEWORK_API FLinearColor GetLineColor(EGizmoElementInteractionState InteractionState) const;
+
 	// Update this line render state based on a child mesh render state attribute.
 	INTERACTIVETOOLSFRAMEWORK_API void Update(FGizmoElementLineRenderStateAttributes& InChildAttributes);
+
+private:
+	// Returns the line color corresponding to the input interaction state, shared between the public const and non-const versions
+	FLinearColor GetLineColorInternal(EGizmoElementInteractionState InteractionState) const;
 };
 
