[ITF] GizmoElementCylinder can be calculated in screen space, affecting dash length, spacing, and line radius

diff --git a/Engine/Source/Runtime/InteractiveToolsFramework/Private/BaseGizmos/GizmoElementCylinder.cpp b/Engine/Source/Runtime/InteractiveToolsFramework/Private/BaseGizmos/GizmoElementCylinder.cpp
--- a/Engine/Source/Runtime/InteractiveToolsFramework/Private/BaseGizmos/GizmoElementCylinder.cpp
+++ b/Engine/Source/Runtime/InteractiveToolsFramework/Private/BaseGizmos/GizmoElementCylinder.cpp
@@ -8,33 +8,134 @@
 #include "DynamicMeshBuilder.h"
 #include "GizmoPrivateUtil.h"
 #include "InputState.h"
+#include "IO/IoDispatcher.h"
 #include "Materials/MaterialInterface.h"
 #include "PrimitiveDrawingUtils.h"
 #include "SceneManagement.h"
 #include "SphereTypes.h"
+#include "ToolDataVisualizer.h"
 
 #include UE_INLINE_GENERATED_CPP_BY_NAME(GizmoElementCylinder)
 
 namespace GizmoElementCylinderLocals
 {
+	constexpr int32 MaxNumDashes = 512; // Arbitrary maximum to prevent overflow in the case of very long cylinders
+
 	/**
-	 * Normalizes the given dash parameters to fit the input range.
-	 * @param	OutScaleFactor the scale factor to apply to the dash length and gap length.
-	 * @return	the number of dashes that fit within the range.
-	 */
-	template < typename RealType = float>
-	int32 GetDashNumAndScale(const RealType& InLineLength, const RealType& InDashLength, const RealType& InDashGapLength, RealType& OutScaleFactor)
+	 * Get the screen-space start and end for the given line (from origin, direction, length). Z is depth from view, XY are in pixel space.
+	 * OutClampedWorldFraction is used when one or more of the end points is beyond the view bounds, and reflects the fraction of the line (in world space) that's visible,
+	 * ie. if only 50% of the line is within the viewport, this will return 0.5. */
+	void GetScreenSpaceLineStartEnd(const FSceneView* InSceneView, const FVector& InLineOrigin, const FVector& InLineDirection, const double& InLineLength, FVector& OutScreenStartPoint, FVector& OutScreenEndPoint, double& OutClampedWorldFraction)
 	{
-		const RealType DashLengthAndGap = InDashLength + InDashGapLength;
-		OutScaleFactor = static_cast<RealType>(1);
+		const FIntRect ViewRect = InSceneView->UnscaledViewRect;
+
+		FVector LineStart = InLineOrigin;
+		FVector LineEnd = InLineOrigin + InLineDirection * InLineLength;
+
+		FVector4 ScreenStartPoint = InSceneView->WorldToScreen(LineStart);
+		FVector4 ScreenEndPoint = InSceneView->WorldToScreen(LineEnd);
+
+		FVector2D UnclampedPixelStartPoint = FVector2D::ZeroVector;
+		InSceneView->ScreenToPixel(ScreenStartPoint, UnclampedPixelStartPoint);
+
+		FVector2D UnclampedPixelEndPoint = FVector2D::ZeroVector;
+		InSceneView->ScreenToPixel(ScreenEndPoint, UnclampedPixelEndPoint);
+
+		auto ClampToViewRect = [&](const FVector2D& InPoint)
+		{
+			// @see: FSceneView::CursorToScreen
+			return FVector2D::Clamp(InPoint, FVector2D(ViewRect.Min.X, ViewRect.Min.Y), FVector2D(ViewRect.Max.X, ViewRect.Max.Y));
+		};
+
+		const FVector2D ClampedPixelStartPoint = ClampToViewRect(UnclampedPixelStartPoint);
+		const FVector2D ClampedPixelEndPoint = ClampToViewRect(UnclampedPixelEndPoint);
+
+		const double UnclampedLineLength = FVector2D::Distance(UnclampedPixelStartPoint, UnclampedPixelEndPoint);
+		const double ClampedLineLength = FVector2D::Distance(ClampedPixelStartPoint, ClampedPixelEndPoint);
+
+		const bool bIsClamped = !FMath::IsNearlyEqual(UnclampedLineLength, ClampedLineLength);
+
+		// If clamped, we need to adjust the Z values to be at the clamped position along the line
+		if (bIsClamped)
+		{
+			// @note: the frustum clipping code below could be cached, with start and end intersection world points which are then used if the corresponding point is out of bounds
+			// @see: LevelEditorViewport.cpp TrimLineToFrustum
+			FVector IntersectionPoint;
+			for (const FPlane& FrustumPlane : InSceneView->ViewFrustum.Planes)
+			{
+				if (FMath::SegmentPlaneIntersection(LineStart, LineEnd, FrustumPlane, IntersectionPoint))
+				{
+					if ((FrustumPlane.Dot(IntersectionPoint - LineEnd) > UE_DOUBLE_SMALL_NUMBER))
+					{
+						LineStart = IntersectionPoint;
+					}
+					else
+					{
+						LineEnd = IntersectionPoint;
+					}
+				}
+			}
 
-		// Number based on sequence of dash, gap pairs, and last dash. This should always be an even number for fit alignment
-		const int32 NumDashes = FMath::Max(2, FMath::CeilToInt32((InLineLength / DashLengthAndGap)));
+			OutClampedWorldFraction = FVector::Distance(LineStart, LineEnd) / InLineLength;
 
-		// Get a scale factor based on how many dashes with gaps (and an end dash) can actually fit on this line
-		OutScaleFactor = InLineLength / (((NumDashes - 1) * DashLengthAndGap) + InDashLength);
+			ScreenStartPoint = InSceneView->WorldToScreen(LineStart);
+			ScreenEndPoint = InSceneView->WorldToScreen(LineEnd);
+		}
 
-		return NumDashes;
+		OutScreenStartPoint = FVector(ClampedPixelStartPoint.X, ClampedPixelStartPoint.Y, ScreenStartPoint.W);
+		OutScreenEndPoint = FVector(ClampedPixelEndPoint.X, ClampedPixelEndPoint.Y, ScreenEndPoint.W);
+
+		// Clip to near plane if needed
+		OutScreenStartPoint.Z = FMath::Max(OutScreenStartPoint.Z, UE_DOUBLE_SMALL_NUMBER);
+		OutScreenEndPoint.Z = FMath::Max(OutScreenEndPoint.Z, UE_DOUBLE_SMALL_NUMBER);
+	}
+
+	double GetPixelToWorldScaleAtScreenFraction(
+		const double& InStartDepthInView, const double& InEndDepthInView,
+		double InScreenFraction)
+	{
+		// Clamp screen fraction to normalized range
+		InScreenFraction = FMath::Clamp(InScreenFraction, 0.0, 1.0);
+
+		// Fraction from the end depth, scaled to the (start - end) range
+		const double ScreenFractionFromEnd = FMath::Max(
+        	(1.0 - InScreenFraction) * InEndDepthInView + (InScreenFraction * InStartDepthInView),
+        	UE_DOUBLE_SMALL_NUMBER);
+
+		// As a ratio, for parameterization square the result
+		return InEndDepthInView / ScreenFractionFromEnd;
+	}
+
+	// Copy/Paste, with added Vertex Color, @see: PrimitiveDrawingUtil.h, DrawCylinder(...);
+	void DrawCylinder(
+		class FPrimitiveDrawInterface* PDI,
+		const FMatrix& CylToWorld,
+		const FVector& Base,
+		const FVector& XAxis,
+		const FVector& YAxis,
+		const FVector& ZAxis,
+		double Radius,
+		double TopRadius,
+		double HalfHeight,
+		uint32 Sides,
+		const FMaterialRenderProxy* MaterialInstance,
+		uint8 DepthPriority,
+		const FColor& InColor)
+	{
+		TArray<FDynamicMeshVertex> MeshVerts;
+		TArray<uint32> MeshIndices;
+		BuildTaperedCylinderVerts(Base, XAxis, YAxis, ZAxis, Radius, TopRadius, HalfHeight, Sides, MeshVerts, MeshIndices);
+
+		Algo::ForEach(MeshVerts, [InColor](FDynamicMeshVertex& InVertex)
+		{
+			InVertex.Color = InColor;
+		});
+
+		FDynamicMeshBuilder MeshBuilder(PDI->View->GetFeatureLevel());
+		MeshBuilder.AddVertices(MeshVerts);
+		MeshBuilder.AddTriangles(MeshIndices);
+
+		MeshBuilder.Draw(PDI, CylToWorld, MaterialInstance, DepthPriority);
 	}
 }
 
@@ -47,61 +148,251 @@ void UGizmoElementCylinder::Render(IToolsContextRenderAPI* RenderAPI, const FRen
 	{
 		if (const UMaterialInterface* UseMaterial = CurrentRenderState.GetCurrentMaterial())
 		{
+			// Line is too small to draw
+			if (FMath::IsNearlyZero(Height))
+			{
+				return;
+			}
+
 			const FQuat Rotation = FRotationMatrix::MakeFromZ(Direction).ToQuat();
 			const double HalfHeight = Height * 0.5;
 			const FVector OriginOffset = Direction * HalfHeight;
 
-			FTransform RenderLocalToWorldTransform = FTransform(Rotation, OriginOffset) * CurrentRenderState.LocalToWorldTransform;
+			const FTransform RenderLocalToWorldTransform = FTransform(Rotation, OriginOffset) * CurrentRenderState.LocalToWorldTransform;
 			FPrimitiveDrawInterface* PDI = RenderAPI->GetPrimitiveDrawInterface();
 
-			if (bIsDashed)
+			const FColor Color = CurrentRenderState.GetCurrentVertexColor().ToFColor(true);
+
+			const FVector WorldOrigin = CurrentRenderState.LocalToWorldTransform.GetLocation();
+			const FVector WorldDirection = CurrentRenderState.LocalToWorldTransform.TransformVectorNoScale(Direction);
+			const double WorldHeight = Height * CurrentRenderState.LocalToWorldTransform.GetScale3D().X;
+			
+			// Get screen-space endpoints (used for visibility and for reparameterization, both world and screen-space)
+            FVector ScreenStartPoint, ScreenEndPoint;
+            double LineFractionInView = 1.0;
+            GizmoElementCylinderLocals::GetScreenSpaceLineStartEnd(RenderAPI->GetSceneView(), WorldOrigin, WorldDirection, WorldHeight, ScreenStartPoint, ScreenEndPoint, LineFractionInView);
+
+			const float HeightInView = Height * LineFractionInView;
+
+			// Screen-space dashed: need to reparameterize screen fractions -> world fractions and scale radii per segment, only used in screen-space path
+			const float ScreenSpaceHeight = FVector2D::Distance(FVector2D(ScreenStartPoint), FVector2D(ScreenEndPoint));
+
+			auto AppendCylinder = [&](
+				const double& InWorldStart, const double& InWorldEnd,
+				const double& InRadiusScaleStart, const double& InRadiusScaleEnd,
+				TArray<FDynamicMeshVertex>& OutVertices, TArray<uint32>& OutTriangleIndices)
 			{
-				// Get a scale factor based on how many dashes with gaps (and an end dash) can actually fit on this line
-				float DashLengthScale = 0.0f;
+				const double WorldDashMid = 0.5 * (InWorldStart + InWorldEnd);
+				const double HalfDashLength = 0.5 * (InWorldEnd - InWorldStart) * HeightInView;
+
+				BuildTaperedCylinderVerts(
+					(-FVector::UpVector * HalfHeight) + FVector::UpVector * (HeightInView * WorldDashMid),
+					FVector::XAxisVector,
+					FVector::YAxisVector,
+					FVector::ZAxisVector,
+					Radius * InRadiusScaleStart,
+					Radius * InRadiusScaleEnd,
+					HalfDashLength,
+					NumSides,
+					OutVertices,
+					OutTriangleIndices);
+			};
+
+			/** Defines unique behavior for building dashes, which differs between world and screen-space. */
+			struct FDashPolicy
+			{
+				/** Total length of the line. */
+				double TotalLength = 0.0;
+
+				/** Converts the given parameter along the length (0-1) to a world fraction. */
+				TFunction<double(const double& InParam)> ParamToWorldFunction;
+
+				/** Returns the radius scale at the given param (0-1) */
+				TFunction<double(const double& InParam)> ParamToRadiusScaleFunction;
+			};
+
+			auto BuildDashes = [&](const FDashPolicy& InPolicy)
+			{
+				// Determine the number of dashes based on world or screen-space metrics
+				const float DashAndGapLength = (DashLength + DashGapLength);
+				const int32 NumDashes = FMath::Max(1, FMath::CeilToInt32(InPolicy.TotalLength / DashAndGapLength));
+				const int32 AdjustedNumDashes = FMath::Min(GizmoElementCylinderLocals::MaxNumDashes, NumDashes);
+
+				const bool bAddFillerDash = NumDashes - AdjustedNumDashes > 0;
+
+				// Whether to add the filler dash before or after the dashes (used so that the closest geometry has dashes vs. the furthest.
+				const bool bAddFillerBeforeDashes = bAddFillerDash && ScreenStartPoint.Z > ScreenEndPoint.Z;
+
+				// Num Sides * 2 (top and bottom)
+				const int32 NumVerticesPerDash = NumSides * 2;
+
+				// Side Quads NumSides - 1 * 2 triangles per quad, Cap Tris (fan) = NumSides per cap
+				const int32 NumTrianglesPerDash = ((NumSides - 1) * 2) + (NumSides * 2);
+
+				const int32 NumCylinders = (AdjustedNumDashes + (bAddFillerDash ? 1 : 0));
+
+				FDynamicMeshBuilder MeshBuilder(PDI->View->GetFeatureLevel());
+				MeshBuilder.ReserveVertices(NumVerticesPerDash * NumCylinders);
+				MeshBuilder.ReserveTriangles(NumTrianglesPerDash * NumCylinders);
+
+				TArray<FDynamicMeshVertex> Vertices;
+				Vertices.Reserve(NumVerticesPerDash * NumCylinders);
+
+				TArray<uint32> TriangleIndices;
+				TriangleIndices.Reserve(NumTrianglesPerDash * NumCylinders);
+
+				auto AppendDash = [&](const double& InStart, const double& InEnd)
+				{
+					const double WorldStart = InPolicy.ParamToWorldFunction(InStart);
+					const double WorldEnd = InPolicy.ParamToWorldFunction(InEnd);
+
+					const double RadiusScaleStart = InPolicy.ParamToRadiusScaleFunction(InStart);
+					const double RadiusScaleEnd = InPolicy.ParamToRadiusScaleFunction(InEnd);
+
+					AppendCylinder(
+						WorldStart,
+						WorldEnd,
+						RadiusScaleStart,
+						RadiusScaleEnd,
+						Vertices,
+						TriangleIndices);
+				};
+
+				const int32 NumDashesForSpaceAllocation = FMath::Max(0, AdjustedNumDashes - 1 + (bAddFillerDash ? 1 : 0));
+				const double SpaceForDashes = DashAndGapLength * NumDashesForSpaceAllocation;
+
+				double UsedLength = 0.0;
+
+				// Filler BEFORE dashes
+				if (bAddFillerDash && bAddFillerBeforeDashes)
+				{
+					constexpr double FillerStart = 0.0;
+					const double FillerEnd = (InPolicy.TotalLength - SpaceForDashes) / InPolicy.TotalLength;
+
+					AppendDash(FillerStart, FillerEnd);
+
+					UsedLength = InPolicy.TotalLength - SpaceForDashes;
+				}
 				
-				constexpr int32 MaxNumDashes = 512; // Arbitrary maximum to prevent overflow in the case of very long cylinders
+				int32 DashIndex = 0;
+				while (UsedLength < InPolicy.TotalLength - 0.5 && DashIndex < AdjustedNumDashes)
+				{
+					// Clamp the dash length to avoid overshoot
+					const double AdjustedDashLength = FMath::Min(DashLength, InPolicy.TotalLength - UsedLength);
+
+					const double DashStart = UsedLength / InPolicy.TotalLength;
+					const double DashEnd = (UsedLength + AdjustedDashLength) / InPolicy.TotalLength;
+
+					AppendDash(DashStart, DashEnd);
+
+					UsedLength += AdjustedDashLength + DashGapLength;
+					++DashIndex;
+				}
+
+				// Filler AFTER dashes
+				if (bAddFillerDash && !bAddFillerBeforeDashes)
+				{
+					const double RemainingLength = InPolicy.TotalLength - UsedLength;
+					if (RemainingLength > UE_DOUBLE_SMALL_NUMBER)
+					{
+						const double FillerStart = UsedLength / InPolicy.TotalLength;
+						constexpr double FillerEnd = 1.0;
+
+						AppendDash(FillerStart, FillerEnd);
+					}
+				}
+
+				// Colorize
+				Algo::ForEach(
+					Vertices,
+					[Color](FDynamicMeshVertex& InVertex)
+					{
+						InVertex.Color = Color;
+					});
+
+				MeshBuilder.AddVertices(Vertices);
+				MeshBuilder.AddTriangles(TriangleIndices);
+
+				MeshBuilder.Draw(PDI, RenderLocalToWorldTransform.ToMatrixWithScale(), UseMaterial->GetRenderProxy(), SDPG_Foreground);
+			};
+
+			if (bIsDashed)
+			{
+				if (bScreenSpace)
+				{
+					FDashPolicy ScreenSpaceDashPolicy;
+					ScreenSpaceDashPolicy.TotalLength = ScreenSpaceHeight;
+					ScreenSpaceDashPolicy.ParamToWorldFunction = [StartDepth = ScreenStartPoint.Z, EndDepth = ScreenEndPoint.Z](const double& InParam)
+					{
+						// Converts a screen-space fraction (0-1) to a world-space fraction (0-1), so that, for example,
+						// the mid-point will appear to correctly to the viewer even if it's not the actual world-space mid-point.
 
-				// Number based on sequence of dash, gap pairs, and last dash
-				const int32 NumDashes = FMath::Min(MaxNumDashes, GizmoElementCylinderLocals::GetDashNumAndScale(FMath::Max(0.0f, Height - DashLength), DashLength, DashGapLength, DashLengthScale));
+						// Clamp screen fraction to normalized range
+						const double ClampedParam = FMath::Clamp(InParam, 0.0, 1.0);
 
-				const float AdjustedDashLength = DashLength * DashLengthScale;
-				const float AdjustDashGapLength = DashGapLength * DashLengthScale;
+						return (ClampedParam * StartDepth) / (((1.0 - ClampedParam) * EndDepth) + (ClampedParam * StartDepth));
+					};
 
-				const float DashLengthAndGaps = (AdjustedDashLength + AdjustDashGapLength);
+					ScreenSpaceDashPolicy.ParamToRadiusScaleFunction = [StartDepth = ScreenStartPoint.Z, EndDepth = ScreenEndPoint.Z](const double& InParam)
+					{
+						return GizmoElementCylinderLocals::GetPixelToWorldScaleAtScreenFraction(StartDepth, EndDepth, InParam);
+					};
 
-				const float HalfDashLength = AdjustedDashLength * 0.5f;
+					BuildDashes(ScreenSpaceDashPolicy);
+				}
+				else
+				{
+					FDashPolicy WorldSpaceDashPolicy;
+					WorldSpaceDashPolicy.TotalLength = HeightInView;
+					WorldSpaceDashPolicy.ParamToWorldFunction = [](const double& InParam) { return InParam; }; // Linear mapping in world-space
+					WorldSpaceDashPolicy.ParamToRadiusScaleFunction = [](const double&) { return 1.0; }; // Constant radius in world-space
 
-				for (int32 DashIndex = 0; DashIndex <= NumDashes; ++DashIndex)
+					BuildDashes(WorldSpaceDashPolicy);
+				}
+			}
+			else // Solid, not dashed
+			{
+				if (bScreenSpace)
 				{
-					DrawCylinder(
+					constexpr double Start = 0.0;
+					constexpr double End = 1.0;
+
+					const double RadiusScaleStart = GizmoElementCylinderLocals::GetPixelToWorldScaleAtScreenFraction(ScreenStartPoint.Z, ScreenEndPoint.Z, Start);
+					const double RadiusScaleEnd = GizmoElementCylinderLocals::GetPixelToWorldScaleAtScreenFraction(ScreenStartPoint.Z, ScreenEndPoint.Z, End);
+
+					GizmoElementCylinderLocals::DrawCylinder(
 						PDI,
 						RenderLocalToWorldTransform.ToMatrixWithScale(),
-						(-FVector::UpVector * HalfHeight) + FVector::UpVector * (HalfDashLength + (DashIndex * DashLengthAndGaps)),
-						FVector(1, 0, 0),
-						FVector(0, 1, 0),
-						FVector(0, 0, 1),
+						(-FVector::UpVector * HalfHeight) + FVector::UpVector * HalfHeight,
+						FVector::XAxisVector,
+						FVector::YAxisVector,
+						FVector::ZAxisVector,
+						Radius * RadiusScaleStart,
+						Radius * RadiusScaleEnd,
+						HalfHeight,
+						NumSides,
+						UseMaterial->GetRenderProxy(),
+						SDPG_Foreground,
+						Color);
+				}
+				else
+				{
+					GizmoElementCylinderLocals::DrawCylinder(
+						PDI,
+						RenderLocalToWorldTransform.ToMatrixWithScale(),
+						(-FVector::UpVector * HalfHeight) + FVector::UpVector * HalfHeight,
+						FVector::XAxisVector,
+						FVector::YAxisVector,
+						FVector::ZAxisVector,
 						Radius,
-						HalfDashLength,
+						Radius,
+						HalfHeight,
 						NumSides,
 						UseMaterial->GetRenderProxy(),
-						SDPG_Foreground);
+						SDPG_Foreground,
+						Color);
 				}
 			}
-			else
-			{
-				DrawCylinder(
-					PDI,
-					RenderLocalToWorldTransform.ToMatrixWithScale(),
-					FVector::ZeroVector,
-					FVector(1, 0, 0),
-					FVector(0, 1, 0),
-					FVector(0, 0, 1),
-					Radius,
-					HalfHeight,
-					NumSides,
-					UseMaterial->GetRenderProxy(),
-					SDPG_Foreground);
-			}
 		}
 	}
 }
@@ -254,3 +545,12 @@ void UGizmoElementCylinder::GetDashParameters(float& OutDashLength, float& OutGa
 	OutGapLength = DashGapLength;
 }
 
+void UGizmoElementCylinder::SetScreenSpace(bool bInScreenSpace)
+{
+	bScreenSpace = bInScreenSpace;
+}
+
+bool UGizmoElementCylinder::GetScreenSpace() const
+{
+	return bScreenSpace;
+}
diff --git a/Engine/Source/Runtime/InteractiveToolsFramework/Public/BaseGizmos/GizmoElementCylinder.h b/Engine/Source/Runtime/InteractiveToolsFramework/Public/BaseGizmos/GizmoElementCylinder.h
--- a/Engine/Source/Runtime/InteractiveToolsFramework/Public/BaseGizmos/GizmoElementCylinder.h
+++ b/Engine/Source/Runtime/InteractiveToolsFramework/Public/BaseGizmos/GizmoElementCylinder.h
@@ -53,6 +53,10 @@ public:
 	INTERACTIVETOOLSFRAMEWORK_API virtual void SetDashParameters(const float InDashLength = 10.0f, const TOptional<float>& InGapLength = TOptional<float>());
 	INTERACTIVETOOLSFRAMEWORK_API virtual void GetDashParameters(float& OutDashLength, float& OutGapLength) const;
 
+	// Whether the radius, dash length and spacing (if applicable) are in screen space
+	INTERACTIVETOOLSFRAMEWORK_API virtual void SetScreenSpace(bool bInScreenSpace);
+	INTERACTIVETOOLSFRAMEWORK_API virtual bool GetScreenSpace() const;
+
 protected:
 
 	// Location of center of cylinder's base circle.
@@ -86,4 +90,8 @@ protected:
 	// The Dash Gap Length, if Dashing is enabled.
 	UPROPERTY()
 	float DashGapLength = 5.0f;
+
+	// Whether the radius, dash length and spacing (if applicable) are in screen space
+	UPROPERTY(Getter = "GetScreenSpace", Setter = "SetScreenSpace")
+	bool bScreenSpace = false;
 };
