[Editor ITF] Added Actor Snapping

diff --git a/Engine/Source/Editor/Experimental/EditorInteractiveToolsFramework/Private/Snapping/EditorSnappingManager.cpp b/Engine/Source/Editor/Experimental/EditorInteractiveToolsFramework/Private/Snapping/EditorSnappingManager.cpp
--- a/Engine/Source/Editor/Experimental/EditorInteractiveToolsFramework/Private/Snapping/EditorSnappingManager.cpp
+++ b/Engine/Source/Editor/Experimental/EditorInteractiveToolsFramework/Private/Snapping/EditorSnappingManager.cpp
@@ -2,45 +2,28 @@
 
 #include "Snapping/EditorSnappingManager.h"
 
+#include "Algo/RemoveIf.h"
+#include "CollisionQueryFilterCallbackCore.h"
+#include "CollisionQueryParams.h"
+#include "Components/PrimitiveComponent.h"
 #include "ContextObjectStore.h"
+#include "DrawDebugHelpers.h"
+#include "Engine/World.h"
 #include "GridSnapping.h"
 #include "InteractiveToolsContext.h"
 #include "PolicySnapping.h"
 #include "SceneQueries/SceneSnappingManager.h"
 #include "ToolContextInterfaces.h"
+#include "TransformSnapping.h"
 #include "VectorUtil.h"
 #include "ViewportSnappingModule.h"
+#include "EditorViewportClient.h"
+#include "SceneView.h"
 
 #include UE_INLINE_GENERATED_CPP_BY_NAME(EditorSnappingManager)
 
 namespace UE::Editor::Gizmos
 {
-	namespace EditorSnappingManagerLocals
-	{
-		// This assumes the provided AxisList only specifies a single axis
-		EAxis::Type GetAxisFromAxisList(EAxisList::Type InAxisList)
-		{
-			if (InAxisList == EAxisList::X)
-            {
-                return EAxis::X;
-            }
-
-			if (InAxisList == EAxisList::Y)
-			{
-				return EAxis::Y;
-			}
-
-			if (InAxisList == EAxisList::Z)
-			{
-				return EAxis::Z;
-			}
-
-			ensureMsgf(true, TEXT("Invalid AxisList, it must specify either X, Y or Z"));
-
-			return EAxis::None;
-		}
-	}
-
 	bool RegisterSceneSnappingManager(UInteractiveToolsContext* InToolsContext)
 	{
 		if (!ensure(InToolsContext))
@@ -97,6 +80,91 @@ namespace UE::Editor::Gizmos
 
 		return nullptr;
 	}
+
+	namespace Private
+	{
+		class FEditorSceneSnappingManagerDebug : public FTickableEditorObject
+		{
+		public:
+			static constexpr bool bDebugEnabled = false;
+
+		public:
+			//~ Begin FTickableEditorObject
+			virtual void Tick(float DeltaTime) override
+			{
+				if (!QueriesAPI || !DebugData.bShouldTick)
+				{
+					return;
+				}
+
+				if (UWorld* World = QueriesAPI->GetCurrentEditingWorld();
+					DebugData.bHasValidHit && World)
+				{
+					const FColor DebugColor = FColor::Cyan;
+					DrawDebugDirectionalArrow(World, DebugData.LineStart, DebugData.LineEnd, 80.0f, DebugColor);
+					DrawDebugPoint(World, DebugData.Point, 10.0f, FColor::Cyan);
+					DrawDebugCrosshairs(World, DebugData.Point, FRotator::ZeroRotator, 1000.0f, FColor::Red);
+
+					FVector Axis0 = DebugData.AxisY;
+					FVector Axis1 = DebugData.AxisX;
+					DrawDebugCircle(World, DebugData.Point, 100.0f, 16, DebugColor, false, -1, 0, 0, Axis0, Axis1);
+
+					// DrawDebugString(World, DebugData.Point, DebugData.HitLabel, nullptr, DebugColor, -1, true, 100);
+					// DrawDebugCircle(World,)
+				}
+			}
+
+			TStatId GetStatId() const
+			{
+				RETURN_QUICK_DECLARE_CYCLE_STAT(UEditorSceneSnappingManager, STATGROUP_Tickables);
+			}
+
+			ETickableTickType GetTickableTickType() const
+			{
+				return ETickableTickType::Conditional;
+			}
+
+			bool IsTickable() const
+			{
+				return DebugData.bShouldTick && FTickableEditorObject::IsTickable();
+			}
+			//~ End FTickableEditorObject
+
+		public:
+			const IToolsContextQueriesAPI* QueriesAPI = nullptr;
+
+			mutable struct FDebugData
+			{
+				std::atomic_bool bShouldTick = false;
+
+				bool bHasValidHit = false;
+				FVector LineStart;
+				FVector LineEnd;
+				FVector Point;
+				FVector Normal;
+				FVector AxisX;
+				FVector AxisY;
+				FString HitLabel;
+
+				void Reset()
+				{
+					bHasValidHit = false;
+					LineStart = FVector::ZeroVector;
+					LineEnd = FVector::ZeroVector;
+					Point = FVector::ZeroVector;
+					Normal = FVector::ForwardVector;
+					AxisX = FVector::RightVector;
+					AxisY = FVector::UpVector;
+					HitLabel = TEXT("");
+				}
+			} DebugData;
+		};
+	}
+}
+
+UEditorSceneSnappingManager::UEditorSceneSnappingManager()
+{
+	Debug = MakeUnique<UE::Editor::Gizmos::Private::FEditorSceneSnappingManagerDebug>();
 }
 
 void UEditorSceneSnappingManager::Initialize(const TObjectPtr<UInteractiveToolsContext>& InToolsContext)
@@ -105,23 +173,136 @@ void UEditorSceneSnappingManager::Initialize(const TObjectPtr<UInteractiveToolsC
 	QueriesAPI = ToolManager ? ToolManager->GetContextQueriesAPI() : nullptr;
 
 	RegisterSnapQueryTargetHandler(MakeUnique<UE::Editor::Gizmos::FGizmoGridSnapper>(InToolsContext));
+	RegisterSnapQueryTargetHandler(MakeUnique<UE::Editor::Gizmos::FGizmoTransformSnapper>(InToolsContext));
 
-	if (const TSharedPtr<ISnappingPolicy>& MergedSnappingPolicy = IViewportSnappingModule::GetSnapManager())
-	{
-		RegisterSnapQueryTargetHandler(MakeUnique<UE::Editor::Gizmos::FGizmoPolicySnapper>(MergedSnappingPolicy));
-	}
+	Debug->QueriesAPI = QueriesAPI;
+	Debug->DebugData.bShouldTick = true;
 }
 
 void UEditorSceneSnappingManager::Shutdown()
 {
+	Debug->DebugData.bShouldTick = false;
+	Debug->QueriesAPI = nullptr;
+
 	QueryTargetHandlers.Reset();
 
 	QueriesAPI = nullptr;
 }
 
+bool UEditorSceneSnappingManager::ExecuteSceneHitQueryPrimitiveTransform(const FSceneHitQueryRequest& InRequest, TArray<FSceneHitQueryResult>& OutResults) const
+{
+	const FCollisionObjectQueryParams ObjectQueryParams(FCollisionObjectQueryParams::AllObjects);
+	FCollisionQueryParams QueryParams = FCollisionQueryParams::DefaultQueryParam;
+	QueryParams.bTraceComplex = true;
+	QueryParams.bReturnFaceIndex = false;
+
+	const UWorld* World = QueriesAPI->GetCurrentEditingWorld();
+	if (!World)
+	{
+		return false;
+	}
+
+	TArray<FHitResult> HitResults;
+	const FVector RayEnd = InRequest.WorldRay.PointAt(HALF_WORLD_MAX);
+	if (World->LineTraceMultiByObjectType(HitResults, (FVector)InRequest.WorldRay.Origin, RayEnd, ObjectQueryParams, QueryParams))
+	{
+		HitResults.SetNum(Algo::RemoveIf(HitResults, [&](const FHitResult& InHitResult)
+		{
+			return !InRequest.VisibilityFilter.IsVisible(InHitResult.Component.Get());
+		}));
+
+		if (HitResults.IsEmpty())
+		{
+			return false;
+		}
+
+		OutResults.SetNum(HitResults.Num());
+		Algo::Transform(HitResults, OutResults, [&](const FHitResult& InHitResult)
+		{
+			FSceneHitQueryResult HitQueryResult;
+			HitQueryResult.HitResult = InHitResult;
+			HitQueryResult.TargetActor = InHitResult.GetActor();
+			HitQueryResult.TargetComponent = InHitResult.GetComponent();
+			HitQueryResult.Position = InHitResult.GetActor() ? InHitResult.GetActor()->GetActorLocation() : FVector(InHitResult.ImpactPoint);
+
+			return HitQueryResult;
+		});
+
+		// Update Debug
+		using namespace UE::Editor::Gizmos::Private;
+		if (FEditorSceneSnappingManagerDebug::bDebugEnabled && !OutResults.IsEmpty())
+		{
+			Algo::StableSort(HitResults, [WorldRay = InRequest.WorldRay](const FHitResult& A, const FHitResult& B)
+			{
+				return FVector::Distance(WorldRay.Origin, A.GetActor()->GetActorLocation()) < FVector::Distance(WorldRay.Origin, B.GetActor()->GetActorLocation());
+			});
+
+			const FHitResult* FirstValidHitResult = Algo::FindBy(
+				HitResults,
+				true,
+				[&](const FHitResult& InHitResult)
+				{
+					return InHitResult.GetActor() != nullptr;
+				});
+
+			Debug->DebugData.Reset();
+
+			if (FirstValidHitResult)
+			{
+				FViewCameraState CameraState;
+				QueriesAPI->GetCurrentViewState(CameraState);
+
+				FPlane HitPlane(FirstValidHitResult->GetActor()->GetActorLocation(), CameraState.Forward());
+
+				FEditorSceneSnappingManagerDebug::FDebugData& DebugData = Debug->DebugData;
+
+				DebugData.bHasValidHit = true;
+				DebugData.Point = FirstValidHitResult->GetActor()->GetActorLocation();
+				DebugData.Normal = HitPlane.GetNormal();
+				DebugData.AxisX = CameraState.Up();
+				DebugData.AxisY = CameraState.Right();
+				DebugData.LineStart = InRequest.WorldRay.Origin;
+				DebugData.LineEnd = FirstValidHitResult->Location;
+				DebugData.HitLabel = FString::Printf(TEXT("Actor: %s"), *FirstValidHitResult->GetActor()->GetActorLabel());
+
+				// UE_LOG(LogTemp, Warning, TEXT("Hit: %s"), *OutResult.HitResult.GetActor()->GetActorLabel());
+			}
+		}
+
+		return true;
+	}
+
+	return false;
+}
+
 bool UEditorSceneSnappingManager::ExecuteSceneHitQuery(const FSceneHitQueryRequest& InRequest, FSceneHitQueryResult& OutResult) const
 {
-	// @todo: implement when needed/supported
+	TArray<FSceneHitQueryResult> HitResults;
+	const bool bResult = ExecuteSceneHitQuery(InRequest, HitResults);
+	if (bResult && !HitResults.IsEmpty())
+	{
+		OutResult = HitResults[0];
+	}
+
+	return bResult;
+}
+
+bool UEditorSceneSnappingManager::ExecuteSceneHitQuery(const FSceneHitQueryRequest& InRequest, TArray<FSceneHitQueryResult>& OutResults) const
+{
+	if (!QueriesAPI)
+	{
+		return false;
+	}
+
+	// Actor/Primitive Transform
+	if (EnumHasAnyFlags(InRequest.TargetTypes, ESceneHitQueryTargetType::PrimitiveTransform))
+	{
+		if (ExecuteSceneHitQueryPrimitiveTransform(InRequest, OutResults))
+		{
+			return true;
+		}
+	}
+
 	return false;
 }
 
@@ -176,15 +357,8 @@ bool UEditorSceneSnappingManager::ExecuteSceneSnapQueryPosition(const FSceneSnap
 	{
 		if (Handler->IsQueryTypeSupported(ESceneSnapQueryType::Position))
 		{
-			FVector SnappedPosition;
-			if (Handler->SnapPosition(InRequest.Position, SnappedPosition) == ESceneSnapQueryTargetResult::Snapped)
+			if (Handler->SnapPosition(InRequest, OutResults) == ESceneSnapQueryTargetResult::Snapped)
 			{
-				FSceneSnapQueryResult SnapResult;
-				SnapResult.TargetType = Handler->GetTargetType();
-				SnapResult.Position = SnappedPosition;
-
-				OutResults.Emplace(SnapResult);
-
 				return true;
 			}
 		}
@@ -206,15 +380,8 @@ bool UEditorSceneSnappingManager::ExecuteSceneSnapQueryRotation(const FSceneSnap
 	{
 		if (Handler->IsQueryTypeSupported(ESceneSnapQueryType::Rotation))
 		{
-			FQuat SnappedRotation;
-			if (Handler->SnapRotation(InRequest.DeltaRotation, SnappedRotation) == ESceneSnapQueryTargetResult::Snapped)
+			if (Handler->SnapRotation(InRequest, OutResults) == ESceneSnapQueryTargetResult::Snapped)
 			{
-				FSceneSnapQueryResult SnapResult;
-				SnapResult.TargetType = Handler->GetTargetType();
-				SnapResult.DeltaRotation = SnappedRotation;
-
-				OutResults.Emplace(SnapResult);
-
 				return true;
 			}
 		}
@@ -231,10 +398,10 @@ bool UEditorSceneSnappingManager::ExecuteSceneSnapQueryRotationAngle(const FScen
 	{
 		return false;
 	}
-	
+
 	// It's assumed the input AxisList only has a single active axis
-	const EAxis::Type Axis = EditorSnappingManagerLocals::GetAxisFromAxisList(InRequest.AxisList);
-	if (Axis == EAxis::None)
+	const EAxis::Type Axis = EAxis::FromAxisList(InRequest.AxisList);
+	if (Axis == EAxis::None && InRequest.AxisList != EAxisList::Screen)
 	{
 		return false;
 	}
@@ -244,10 +411,10 @@ bool UEditorSceneSnappingManager::ExecuteSceneSnapQueryRotationAngle(const FScen
 		if (Handler->IsQueryTypeSupported(ESceneSnapQueryType::RotationAngle))
 		{
 			double SnappedAngle;
-			if (Handler->SnapRotationAxisAngle(InRequest.RotationAngle, SnappedAngle, Axis) == ESceneSnapQueryTargetResult::Snapped)
+			if (Handler->SnapRotationAxisAngle(InRequest.RotationAngle, SnappedAngle, InRequest.AxisList) == ESceneSnapQueryTargetResult::Snapped)
 			{
 				FSceneSnapQueryResult SnapResult;
-				SnapResult.TargetType = Handler->GetTargetType();
+				SnapResult.TargetType = Handler->GetTargetTypes();
 				SnapResult.RotationAngle = SnappedAngle;
 
 				OutResults.Emplace(SnapResult);
@@ -273,15 +440,8 @@ bool UEditorSceneSnappingManager::ExecuteSceneSnapQueryScale(const FSceneSnapQue
 	{
 		if (Handler->IsQueryTypeSupported(ESceneSnapQueryType::Scale))
 		{
-			FVector SnappedScale;
-			if (Handler->SnapScale(InRequest.Scale, SnappedScale) == ESceneSnapQueryTargetResult::Snapped)
+			if (Handler->SnapScale(InRequest, OutResults) == ESceneSnapQueryTargetResult::Snapped)
 			{
-				FSceneSnapQueryResult SnapResult;
-				SnapResult.TargetType = Handler->GetTargetType();
-				SnapResult.Scale = SnappedScale;
-
-				OutResults.Emplace(SnapResult);
-
 				return true;
 			}
 		}
diff --git a/Engine/Source/Editor/Experimental/EditorInteractiveToolsFramework/Private/Snapping/EditorSnappingUtil.h b/Engine/Source/Editor/Experimental/EditorInteractiveToolsFramework/Private/Snapping/EditorSnappingUtil.h
new file mode 100644
--- /dev/null
+++ b/Engine/Source/Editor/Experimental/EditorInteractiveToolsFramework/Private/Snapping/EditorSnappingUtil.h
@@ -0,0 +1,37 @@
+// Copyright Epic Games, Inc. All Rights Reserved.
+
+#pragma once
+
+#include "CoreTypes.h"
+#include "GameFramework/Actor.h"
+#include "Components/ActorComponent.h"
+#include "ToolContextInterfaces.h"
+
+namespace UE::EditorSnappingUtil
+{
+	/**
+	 * Utility function to get all components of the given type from one or more Actors.
+	 * @return true if any components found.
+	 */
+	template <typename ComponentType UE_REQUIRES(std::is_base_of_v<UActorComponent, std::decay_t<ComponentType>>)>
+	bool GetActorComponents(const TConstArrayView<AActor*>& InActors, TArray<const ComponentType*>& OutComponents)
+	{
+		OutComponents.Reset();
+		OutComponents.Reserve(InActors.Num());
+
+		for (const AActor* Actor : InActors)
+		{
+			if (!Actor)
+			{
+				continue;
+			}
+
+			TArray<const ComponentType*> ComponentsInActor;
+			Actor->GetComponents<const ComponentType>(ComponentsInActor, true);
+
+			OutComponents.Append(ComponentsInActor);
+		}
+
+		return !OutComponents.IsEmpty();
+	}
+}
diff --git a/Engine/Source/Editor/Experimental/EditorInteractiveToolsFramework/Private/Snapping/GridSnapping.cpp b/Engine/Source/Editor/Experimental/EditorInteractiveToolsFramework/Private/Snapping/GridSnapping.cpp
--- a/Engine/Source/Editor/Experimental/EditorInteractiveToolsFramework/Private/Snapping/GridSnapping.cpp
+++ b/Engine/Source/Editor/Experimental/EditorInteractiveToolsFramework/Private/Snapping/GridSnapping.cpp
@@ -51,9 +51,6 @@ namespace UE::Editor::Gizmos
 		if (const FToolContextSnappingConfiguration SnappingConfig = QueriesAPI->GetCurrentSnappingSettings();
 			SnappingConfig.bEnablePositionGridSnapping)
 		{
-			FSceneSnapQueryResult SnapResult;
-			SnapResult.TargetType = ESceneSnapQueryTargetType::Grid;
-
 			FGizmoGridSnapperRequest GridRequest;
 			GridRequest.PositionGridSize = SnappingConfig.PositionGridDimensions;
 
@@ -63,18 +60,23 @@ namespace UE::Editor::Gizmos
 		return ESceneSnapQueryTargetResult::Disabled;
 	}
 
-	ESceneSnapQueryTargetResult FGizmoGridSnapper::SnapPositionAxis(const FGizmoGridSnapperRequest& InTargetRequest, const double InAxisValue, double& OutSnappedAxisValue, const EAxis::Type InAxis) const
+	ESceneSnapQueryTargetResult FGizmoGridSnapper::SnapPositionAxis(const FGizmoGridSnapperRequest& InTargetRequest, const double InAxisValue, double& OutSnappedAxisValue, const EAxisList::Type InAxis) const
 	{
 		OutSnappedAxisValue = InAxisValue;
 
+		if (!ensure(InAxis != EAxisList::Screen))
+		{
+			return ESceneSnapQueryTargetResult::Unsupported;
+		}
+
 		const FVector GridSize = InTargetRequest.PositionGridSize;
 
-		OutSnappedAxisValue = ::GizmoMath::SnapToIncrement(InAxisValue, GridSize.GetComponentForAxis(InAxis));
+		OutSnappedAxisValue = ::GizmoMath::SnapToIncrement(InAxisValue, GridSize.GetComponentForAxis(EAxis::FromAxisList(InAxis)));
 
 		return ESceneSnapQueryTargetResult::Snapped;
 	}
 
-	ESceneSnapQueryTargetResult FGizmoGridSnapper::SnapPositionAxis(const double InAxisValue, double& OutSnappedAxisValue, const EAxis::Type InAxis) const
+	ESceneSnapQueryTargetResult FGizmoGridSnapper::SnapPositionAxis(const double InAxisValue, double& OutSnappedAxisValue, const EAxisList::Type InAxis) const
 	{
 		OutSnappedAxisValue = InAxisValue;
 
@@ -86,9 +88,6 @@ namespace UE::Editor::Gizmos
 		if (const FToolContextSnappingConfiguration SnappingConfig = QueriesAPI->GetCurrentSnappingSettings();
 			SnappingConfig.bEnablePositionGridSnapping)
 		{
-			FSceneSnapQueryResult SnapResult;
-			SnapResult.TargetType = ESceneSnapQueryTargetType::Grid;
-
 			FGizmoGridSnapperRequest GridRequest;
 			GridRequest.PositionGridSize = SnappingConfig.PositionGridDimensions;
 
@@ -139,9 +138,6 @@ namespace UE::Editor::Gizmos
 		if (const FToolContextSnappingConfiguration SnappingConfig = QueriesAPI->GetCurrentSnappingSettings();
 			SnappingConfig.bEnableRotationGridSnapping)
 		{
-			FSceneSnapQueryResult SnapResult;
-			SnapResult.TargetType = ESceneSnapQueryTargetType::Grid;
-
 			FGizmoGridSnapperRequest GridRequest;
 			GridRequest.RotationGridSize = SnappingConfig.RotationGridAngles;
 
@@ -151,13 +147,14 @@ namespace UE::Editor::Gizmos
 		return ESceneSnapQueryTargetResult::Disabled;
 	}
 
-	ESceneSnapQueryTargetResult FGizmoGridSnapper::SnapRotationAxisAngle(const FGizmoGridSnapperRequest& InTargetRequest, const double& InAngle, double& OutSnappedAngle, const EAxis::Type InAxis) const
+	ESceneSnapQueryTargetResult FGizmoGridSnapper::SnapRotationAxisAngle(const FGizmoGridSnapperRequest& InTargetRequest, const double& InAngle, double& OutSnappedAngle, const EAxisList::Type InAxis) const
 	{
 		OutSnappedAngle = InAngle;
 
 		const FRotator GridSize = InTargetRequest.RotationGridSize;
 
-		const double SnapDelta = GridSize.GetComponentForAxis(InAxis);
+		// If we're in screen space, just use the first component (pitch)
+		const double SnapDelta = InAxis == EAxisList::Screen ? GridSize.Pitch : GridSize.GetComponentForAxis(EAxis::FromAxisList(InAxis));
 
 		OutSnappedAngle = FMath::RadiansToDegrees(InAngle);
 		OutSnappedAngle = ::GizmoMath::SnapToIncrement(OutSnappedAngle, SnapDelta);
@@ -166,7 +163,7 @@ namespace UE::Editor::Gizmos
 		return ESceneSnapQueryTargetResult::Snapped;
 	}
 
-	ESceneSnapQueryTargetResult FGizmoGridSnapper::SnapRotationAxisAngle(const double& InAngle, double& OutSnappedAngle, const EAxis::Type InAxis) const
+	ESceneSnapQueryTargetResult FGizmoGridSnapper::SnapRotationAxisAngle(const double& InAngle, double& OutSnappedAngle, const EAxisList::Type InAxis) const
 	{
 		OutSnappedAngle = InAngle;
 		
@@ -214,6 +211,8 @@ namespace UE::Editor::Gizmos
 
 	ESceneSnapQueryTargetResult FGizmoGridSnapper::SnapScale(const FVector& InScale, FVector& OutSnappedScale, const EAxisList::Type InAxisList) const
 	{
+		OutSnappedScale = InScale;
+
 		if (!QueriesAPI)
 		{
 			return ESceneSnapQueryTargetResult::Unsupported;
@@ -234,18 +233,23 @@ namespace UE::Editor::Gizmos
 		return ESceneSnapQueryTargetResult::Disabled;
 	}
 
-	ESceneSnapQueryTargetResult FGizmoGridSnapper::SnapScaleAxis(const FGizmoGridSnapperRequest& InTargetRequest, const double InAxisValue, double& OutSnappedAxisValue, const EAxis::Type InAxis) const
+	ESceneSnapQueryTargetResult FGizmoGridSnapper::SnapScaleAxis(const FGizmoGridSnapperRequest& InTargetRequest, const double InAxisValue, double& OutSnappedAxisValue, const EAxisList::Type InAxis) const
 	{
 		OutSnappedAxisValue = InAxisValue;
 
+		if (!ensure(InAxis != EAxisList::Screen))
+		{
+			return ESceneSnapQueryTargetResult::Unsupported;
+		}
+
 		const FVector GridSize = InTargetRequest.ScaleGridSize;
 
-		OutSnappedAxisValue = ::GizmoMath::SnapToIncrement(InAxisValue, GridSize.GetComponentForAxis(InAxis));
+		OutSnappedAxisValue = ::GizmoMath::SnapToIncrement(InAxisValue, GridSize.GetComponentForAxis(EAxis::FromAxisList(InAxis)));
 
 		return ESceneSnapQueryTargetResult::Snapped;
 	}
 
-	ESceneSnapQueryTargetResult FGizmoGridSnapper::SnapScaleAxis(const double InAxisValue, double& OutSnappedAxisValue, const EAxis::Type InAxis) const
+	ESceneSnapQueryTargetResult FGizmoGridSnapper::SnapScaleAxis(const double InAxisValue, double& OutSnappedAxisValue, const EAxisList::Type InAxis) const
 	{
 		OutSnappedAxisValue = InAxisValue;
 
diff --git a/Engine/Source/Editor/Experimental/EditorInteractiveToolsFramework/Private/Snapping/GridSnapping.h b/Engine/Source/Editor/Experimental/EditorInteractiveToolsFramework/Private/Snapping/GridSnapping.h
--- a/Engine/Source/Editor/Experimental/EditorInteractiveToolsFramework/Private/Snapping/GridSnapping.h
+++ b/Engine/Source/Editor/Experimental/EditorInteractiveToolsFramework/Private/Snapping/GridSnapping.h
@@ -31,24 +31,24 @@ namespace UE::Editor::Gizmos
 		explicit FGizmoGridSnapper(const UInteractiveToolsContext* InToolsContext);
 		virtual ~FGizmoGridSnapper() override = default;
 
-		//~ Begin FSceneSnapPolicyTarget
+		//~ Begin FSceneSnapPolicyTargetHandler
 		virtual ESceneSnapQueryTargetResult SnapPosition(const FVector& InPosition, FVector& OutSnappedPosition, const EAxisList::Type InAxisList = EAxisList::All) const override;
-		virtual ESceneSnapQueryTargetResult SnapPositionAxis(const double InAxisValue, double& OutSnappedAxisValue, const EAxis::Type InAxis) const override;
+		virtual ESceneSnapQueryTargetResult SnapPositionAxis(const double InAxisValue, double& OutSnappedAxisValue, const EAxisList::Type InAxis) const override;
 		virtual ESceneSnapQueryTargetResult SnapRotation(const FQuat& InRotation, FQuat& OutSnappedRotation, const EAxisList::Type InAxisList = EAxisList::All) const override;
-		virtual ESceneSnapQueryTargetResult SnapRotationAxisAngle(const double& InAngle, double& OutSnappedAngle, const EAxis::Type InAxis) const override;
+		virtual ESceneSnapQueryTargetResult SnapRotationAxisAngle(const double& InAngle, double& OutSnappedAngle, const EAxisList::Type InAxis) const override;
 		virtual ESceneSnapQueryTargetResult SnapScale(const FVector& InScale, FVector& OutSnappedScale, const EAxisList::Type InAxisList = EAxisList::All) const override;
-		virtual ESceneSnapQueryTargetResult SnapScaleAxis(const double InAxisValue, double& OutSnappedAxisValue, const EAxis::Type InAxis) const override;
+		virtual ESceneSnapQueryTargetResult SnapScaleAxis(const double InAxisValue, double& OutSnappedAxisValue, const EAxisList::Type InAxis) const override;
 		virtual bool IsQueryTypeSupported(const ESceneSnapQueryType InQueryType) const override;
-		//~ End FSceneSnapPolicyTarget
+		//~ End FSceneSnapPolicyTargetHandler
 
-		//~ Begin TSceneSnapQueryTarget
+		//~ Begin TSceneSnapQueryTargetHandler
 		virtual ESceneSnapQueryTargetResult SnapPosition(const FGizmoGridSnapperRequest& InTargetRequest, const FVector& InPosition, FVector& OutSnappedPosition, const EAxisList::Type InAxisList = EAxisList::All) const override;
-		virtual ESceneSnapQueryTargetResult SnapPositionAxis(const FGizmoGridSnapperRequest& InTargetRequest, const double InAxisValue, double& OutSnappedPositionAxis, const EAxis::Type InAxis) const override;
+		virtual ESceneSnapQueryTargetResult SnapPositionAxis(const FGizmoGridSnapperRequest& InTargetRequest, const double InAxisValue, double& OutSnappedPositionAxis, const EAxisList::Type InAxis) const override;
 		virtual ESceneSnapQueryTargetResult SnapRotation(const FGizmoGridSnapperRequest& InTargetRequest, const FQuat& InRotation, FQuat& OutSnappedRotation, const EAxisList::Type InAxisList = EAxisList::All) const override;
-		virtual ESceneSnapQueryTargetResult SnapRotationAxisAngle(const FGizmoGridSnapperRequest& InTargetRequest, const double& InAngle, double& OutSnappedAngle, const EAxis::Type InAxis) const override;
+		virtual ESceneSnapQueryTargetResult SnapRotationAxisAngle(const FGizmoGridSnapperRequest& InTargetRequest, const double& InAngle, double& OutSnappedAngle, const EAxisList::Type InAxis) const override;
 		virtual ESceneSnapQueryTargetResult SnapScale(const FGizmoGridSnapperRequest& InTargetRequest, const FVector& InScale, FVector& OutSnappedScale, const EAxisList::Type InAxisList = EAxisList::All) const override;
-		virtual ESceneSnapQueryTargetResult SnapScaleAxis(const FGizmoGridSnapperRequest& InTargetRequest, const double InAxisValue, double& OutSnappedAxisValue, const EAxis::Type InAxis) const override;
-		//~ End TSceneSnapQueryTarget
+		virtual ESceneSnapQueryTargetResult SnapScaleAxis(const FGizmoGridSnapperRequest& InTargetRequest, const double InAxisValue, double& OutSnappedAxisValue, const EAxisList::Type InAxis) const override;
+		//~ End TSceneSnapQueryTargetHandler
 
 	private:
 		const IToolsContextQueriesAPI* QueriesAPI = nullptr;
diff --git a/Engine/Source/Editor/Experimental/EditorInteractiveToolsFramework/Private/Snapping/PolicySnapping.cpp b/Engine/Source/Editor/Experimental/EditorInteractiveToolsFramework/Private/Snapping/PolicySnapping.cpp
--- a/Engine/Source/Editor/Experimental/EditorInteractiveToolsFramework/Private/Snapping/PolicySnapping.cpp
+++ b/Engine/Source/Editor/Experimental/EditorInteractiveToolsFramework/Private/Snapping/PolicySnapping.cpp
@@ -14,29 +14,32 @@ namespace UE::Editor::Gizmos
 		TargetName = "Policy";
 	}
 
-	ESceneSnapQueryTargetResult FGizmoPolicySnapper::SnapTransform(const FTransform& InTransform, FVector& OutSnappedTransform) const
-	{
-		return FSceneSnapQueryTargetHandler::SnapTransform(InTransform, OutSnappedTransform);
-	}
-
 	ESceneSnapQueryTargetResult FGizmoPolicySnapper::SnapPosition(const FVector& InPosition, FVector& OutSnappedPosition, const EAxisList::Type InAxisList) const
 	{
 		OutSnappedPosition = InPosition;
+
 		const FVector GridSize = FVector(GEditor->GetGridSize());
 		SnappingPolicy->SnapPointToGrid(OutSnappedPosition, GridSize);
 
 		return ESceneSnapQueryTargetResult::Snapped;
 	}
 
-	ESceneSnapQueryTargetResult FGizmoPolicySnapper::SnapPositionAxis(const double InAxisDelta, double& OutSnappedAxisDelta, const EAxis::Type InAxis) const
+	ESceneSnapQueryTargetResult FGizmoPolicySnapper::SnapPositionAxis(const double InAxisDelta, double& OutSnappedAxisDelta, const EAxisList::Type InAxis) const
 	{
+		if (InAxis == EAxisList::Screen)
+		{
+			return ESceneSnapQueryTargetResult::Unsupported;
+		}
+
+		const EAxis::Type Axis = EAxis::FromAxisList(InAxis); 
+
 		FVector SnappedPosition;
-		SnappedPosition.SetComponentForAxis(InAxis, InAxisDelta);
+		SnappedPosition.SetComponentForAxis(Axis, InAxisDelta);
 
 		const FVector GridSize = FVector(GEditor->GetGridSize());
 		SnappingPolicy->SnapPointToGrid(SnappedPosition, GridSize);
 
-		OutSnappedAxisDelta = SnappedPosition.GetComponentForAxis(InAxis);
+		OutSnappedAxisDelta = SnappedPosition.GetComponentForAxis(Axis);
 
 		return ESceneSnapQueryTargetResult::Snapped;
 	}
@@ -51,14 +54,16 @@ namespace UE::Editor::Gizmos
 		return ESceneSnapQueryTargetResult::Snapped;
 	}
 
-	ESceneSnapQueryTargetResult FGizmoPolicySnapper::SnapRotationAxisAngle(const double& InAngleDelta, double& OutSnappedAngleDelta, const EAxis::Type InAxis) const
+	ESceneSnapQueryTargetResult FGizmoPolicySnapper::SnapRotationAxisAngle(const double& InAngleDelta, double& OutSnappedAngleDelta, const EAxisList::Type InAxis) const
 	{
+		const EAxis::Type Axis = EAxis::FromAxisList(InAxis); 
+		
 		FRotator Rotation;
-		Rotation.SetComponentForAxis(InAxis, InAngleDelta);
+		Rotation.SetComponentForAxis(Axis, InAngleDelta);
 
 		SnappingPolicy->SnapRotatorToGrid(Rotation);
 
-		OutSnappedAngleDelta = Rotation.GetComponentForAxis(InAxis);
+		OutSnappedAngleDelta = Rotation.GetComponentForAxis(Axis);
 
 		return ESceneSnapQueryTargetResult::Snapped;
 	}
@@ -66,21 +71,29 @@ namespace UE::Editor::Gizmos
 	ESceneSnapQueryTargetResult FGizmoPolicySnapper::SnapScale(const FVector& InScale, FVector& OutSnappedScale, const EAxisList::Type InAxisList) const
 	{
 		OutSnappedScale = InScale;
+
 		const FVector GridSize = FVector(GEditor->GetScaleGridSize());
 		SnappingPolicy->SnapScale(OutSnappedScale, GridSize);
 
 		return ESceneSnapQueryTargetResult::Snapped;
 	}
 
-	ESceneSnapQueryTargetResult FGizmoPolicySnapper::SnapScaleAxis(const double InAxisDelta, double& OutSnappedAxisDelta, const EAxis::Type InAxis) const
+	ESceneSnapQueryTargetResult FGizmoPolicySnapper::SnapScaleAxis(const double InAxisDelta, double& OutSnappedAxisDelta, const EAxisList::Type InAxis) const
 	{
+		if (InAxis == EAxisList::Screen)
+		{
+			return ESceneSnapQueryTargetResult::Unsupported;
+		}
+
+		const EAxis::Type Axis = EAxis::FromAxisList(InAxis); 
+		
 		FVector SnappedScale;
-		SnappedScale.SetComponentForAxis(InAxis, InAxisDelta);
+		SnappedScale.SetComponentForAxis(Axis, InAxisDelta);
 
 		const FVector GridSize = FVector(GEditor->GetScaleGridSize());
 		SnappingPolicy->SnapScale(SnappedScale, GridSize);
 
-		OutSnappedAxisDelta = SnappedScale.GetComponentForAxis(InAxis);
+		OutSnappedAxisDelta = SnappedScale.GetComponentForAxis(Axis);
 
 		return ESceneSnapQueryTargetResult::Snapped;
 	}
diff --git a/Engine/Source/Editor/Experimental/EditorInteractiveToolsFramework/Private/Snapping/PolicySnapping.h b/Engine/Source/Editor/Experimental/EditorInteractiveToolsFramework/Private/Snapping/PolicySnapping.h
--- a/Engine/Source/Editor/Experimental/EditorInteractiveToolsFramework/Private/Snapping/PolicySnapping.h
+++ b/Engine/Source/Editor/Experimental/EditorInteractiveToolsFramework/Private/Snapping/PolicySnapping.h
@@ -26,13 +26,12 @@ namespace UE::Editor::Gizmos
 		virtual ~FGizmoPolicySnapper() override = default;
 
 		//~ Begin FSceneSnapPolicyTarget
-		virtual ESceneSnapQueryTargetResult SnapTransform(const FTransform& InTransform, FVector& OutSnappedTransform) const override;
 		virtual ESceneSnapQueryTargetResult SnapPosition(const FVector& InPosition, FVector& OutSnappedPosition, const EAxisList::Type InAxisList) const override;
-		virtual ESceneSnapQueryTargetResult SnapPositionAxis(const double InAxisDelta, double& OutSnappedAxisDelta, const EAxis::Type InAxis) const override;
+		virtual ESceneSnapQueryTargetResult SnapPositionAxis(const double InAxisDelta, double& OutSnappedAxisDelta, const EAxisList::Type InAxis) const override;
 		virtual ESceneSnapQueryTargetResult SnapRotation(const FQuat& InRotation, FQuat& OutRotation, const EAxisList::Type InAxisList) const override;
-		virtual ESceneSnapQueryTargetResult SnapRotationAxisAngle(const double& InAngleDelta, double& OutSnappedAngleDelta, const EAxis::Type InAxis) const override;
+		virtual ESceneSnapQueryTargetResult SnapRotationAxisAngle(const double& InAngleDelta, double& OutSnappedAngleDelta, const EAxisList::Type InAxis) const override;
 		virtual ESceneSnapQueryTargetResult SnapScale(const FVector& InScale, FVector& OutSnappedScale, const EAxisList::Type InAxisList) const override;
-		virtual ESceneSnapQueryTargetResult SnapScaleAxis(const double InAxisDelta, double& OutSnappedAxisDelta, const EAxis::Type InAxis) const override;
+		virtual ESceneSnapQueryTargetResult SnapScaleAxis(const double InAxisDelta, double& OutSnappedAxisDelta, const EAxisList::Type InAxis) const override;
 		virtual bool IsQueryTypeSupported(const ESceneSnapQueryType InQueryType) const override;
 		//~ End FSceneSnapPolicyTarget
 
diff --git a/Engine/Source/Editor/Experimental/EditorInteractiveToolsFramework/Private/Snapping/TransformSnapping.cpp b/Engine/Source/Editor/Experimental/EditorInteractiveToolsFramework/Private/Snapping/TransformSnapping.cpp
new file mode 100644
--- /dev/null
+++ b/Engine/Source/Editor/Experimental/EditorInteractiveToolsFramework/Private/Snapping/TransformSnapping.cpp
@@ -0,0 +1,288 @@
+// Copyright Epic Games, Inc. All Rights Reserved.
+
+#include "TransformSnapping.h"
+
+#include "BaseGizmos/GizmoMath.h"
+#include "BaseGizmos/GizmoViewContext.h"
+#include "Components/PrimitiveComponent.h"
+#include "ContextObjectStore.h"
+#include "Editor/Experimental/EditorInteractiveToolsFramework/Internal/EditorGizmos/EditorGizmoMath.h"
+#include "EditorSnappingUtil.h"
+#include "GameFramework/Actor.h"
+#include "InteractiveToolsContext.h"
+#include "SceneQueries/SceneSnappingManager.h"
+#include "ToolContextInterfaces.h"
+
+namespace UE::Editor::Gizmos
+{
+	FGizmoTransformSnapper::FGizmoTransformSnapper(const UInteractiveToolsContext* InToolsContext)
+	{
+		TargetName = "ObjectTransform";
+
+		if (ToolManager = InToolsContext->ToolManager;
+			ToolManager.IsValid())
+		{
+			QueriesAPI = ToolManager->GetContextQueriesAPI();
+		}
+	}
+
+	ESceneSnapQueryTargetResult FGizmoTransformSnapper::SnapPosition(const FGizmoTransformSnapperRequest& InTargetRequest, const FVector& InPosition, FVector& OutSnappedPosition, const EAxisList::Type InAxisList) const
+	{
+		FVector WorldDelta = InPosition;
+		if (InTargetRequest.CoordinateSpace != EToolContextCoordinateSystem::World)
+		{
+			WorldDelta = InTargetRequest.CoordinateTransform.GetRotation() * WorldDelta;
+		}
+		OutSnappedPosition = WorldDelta;
+
+		if (!QueriesAPI || !ToolManager.IsValid())
+		{
+			return ESceneSnapQueryTargetResult::Unsupported;
+		}
+
+		UGizmoViewContext* ViewContext = ToolManager->GetContextObjectStore()->FindContext<UGizmoViewContext>();
+		if (!ViewContext)
+		{
+			return ESceneSnapQueryTargetResult::Unsupported;
+		}
+
+		const bool bSnapX = InAxisList & EAxisList::X || InAxisList & EAxisList::Screen;
+		const bool bSnapY = InAxisList & EAxisList::Y || InAxisList & EAxisList::Screen;
+		const bool bSnapZ = InAxisList & EAxisList::Z || InAxisList & EAxisList::Screen;
+
+		if (!bSnapX && !bSnapY && !bSnapZ)
+		{
+			return ESceneSnapQueryTargetResult::NotSnapped;
+		}
+
+		FToolBuilderState ToolBuilderState;
+		QueriesAPI->GetCurrentSelectionState(ToolBuilderState);
+
+		TArray<AActor*> SelectedActors = ToolBuilderState.SelectedActors;
+
+		TArray<const UPrimitiveComponent*> ComponentsToIgnore;
+		UE::EditorSnappingUtil::GetActorComponents(SelectedActors, ComponentsToIgnore);
+
+		FSceneQueryVisibilityFilter HitFilter(&ComponentsToIgnore);
+
+		// Whether to use the gizmo position for uniform axis only (and cursor location for axis, planar).
+		constexpr bool bUseDirectionToGizmoForUniformOnly = true;
+		const bool bUseDirectionToGizmo = bUseDirectionToGizmoForUniformOnly
+			? (InAxisList == EAxisList::All || InAxisList == EAxisList::Screen)
+			: true;
+
+		// Get absolute position (base + delta)
+		const FVector QueryPosition = InTargetRequest.Transform.GetLocation() + WorldDelta;
+
+		// Ray according to the view direction and gizmo position, or cursor position if not uniform axis translation.
+		FVector RayDirection = bUseDirectionToGizmo ? QueryPosition - InTargetRequest.WorldRay.Origin : InTargetRequest.WorldRay.Direction;
+
+		// Always fallback to input direction if coincident.
+		if (RayDirection.IsNearlyZero())
+		{
+			RayDirection = InTargetRequest.WorldRay.Direction;
+		}
+		else
+		{
+			RayDirection = RayDirection.GetSafeNormal();
+		}
+
+		FRay Ray = FRay(InTargetRequest.WorldRay.Origin, RayDirection, true);
+
+		FSceneHitQueryRequest HitRequest(Ray, ESceneHitQueryTargetType::PrimitiveTransform, {}, false, HitFilter);
+
+		TArray<FSceneHitQueryResult> HitResults;
+		if (USceneSnappingManager* SnappingManager = ToolManager->GetContextObjectStore()->FindContext<USceneSnappingManager>())
+		{
+			if (SnappingManager->ExecuteSceneHitQuery(HitRequest, HitResults);
+				!HitResults.IsEmpty())
+			{
+				// Get best hit result by screen space proximity to gizmo
+				const FVector2D ScreenPosition = ViewContext->WorldToPixel(bUseDirectionToGizmo ? QueryPosition : Ray.PointAt(UE_LARGE_HALF_WORLD_MAX1));
+				Algo::StableSort(HitResults, [ScreenPosition, ViewContext](const FSceneHitQueryResult& InHitResultA, const FSceneHitQueryResult& InHitResultB)
+				{
+					const AActor* ActorA = InHitResultA.TargetActor;
+					const AActor* ActorB = InHitResultB.TargetActor;
+
+					if (ActorA == ActorB)
+					{
+						return false;
+					}
+
+					if (ActorA == nullptr)
+					{
+						return false;
+					}
+					
+					if (ActorB == nullptr)
+					{
+						return true;
+					}
+
+					const FVector ActorLocationA = ActorA->GetActorLocation();
+					const FVector ActorLocationB = ActorB->GetActorLocation();
+
+					const FVector2D ActorScreenLocationA = ViewContext->WorldToPixel(ActorLocationA);
+					const FVector2D ActorScreenLocationB = ViewContext->WorldToPixel(ActorLocationB);
+
+					return FVector2D::Distance(ScreenPosition, ActorScreenLocationA) < FVector2D::Distance(ScreenPosition, ActorScreenLocationB);
+				});
+
+				const FSceneHitQueryResult* FirstHitResult = &HitResults[0];
+				if (FirstHitResult->TargetActor)
+				{
+					const FVector2D::FReal DistanceToSnappedTransform = FVector2D::Distance(ScreenPosition, ViewContext->WorldToPixel(FirstHitResult->TargetActor->GetActorLocation()));
+					if (DistanceToSnappedTransform <= InTargetRequest.SnapDistance)
+					{
+						FTransform TargetTransform = FirstHitResult->TargetActor->GetActorTransform();
+						FVector TargetLocation = TargetTransform.GetLocation();
+
+						FVector Delta = TargetLocation - InTargetRequest.Transform.GetLocation();
+
+						if (InTargetRequest.CoordinateSpace != EToolContextCoordinateSystem::World)
+						{
+							Delta = InTargetRequest.CoordinateTransform.InverseTransformVectorNoScale(Delta);
+						}
+
+						if (InAxisList != EAxisList::All && InAxisList != EAxisList::Screen)
+						{
+							const FVector AxisDeltaMultiplier = GizmoMath::GetAxisMultiplier(InAxisList);
+							Delta *= AxisDeltaMultiplier;
+						}
+
+						if (InTargetRequest.CoordinateSpace != EToolContextCoordinateSystem::World)
+						{
+							Delta = InTargetRequest.CoordinateTransform.TransformVectorNoScale(Delta);
+						}
+
+						// Convert absolute back to delta
+						OutSnappedPosition = Delta;
+
+						return ESceneSnapQueryTargetResult::Snapped;
+					}
+				}
+			}
+		}
+
+		return ESceneSnapQueryTargetResult::NotSnapped;
+	}
+
+	ESceneSnapQueryTargetResult FGizmoTransformSnapper::SnapPosition(
+		const FSceneSnapQueryRequest& InRequest,
+		TArray<FSceneSnapQueryResult>& OutResults) const
+	{
+		if (!QueriesAPI)
+		{
+			return ESceneSnapQueryTargetResult::Unsupported;
+		}
+
+		if (const FToolContextSnappingConfiguration SnappingConfig = QueriesAPI->GetCurrentSnappingSettings();
+			SnappingConfig.ObjectTransform.bEnable)
+		{
+			FGizmoTransformSnapperRequest SnapRequest;
+			SnapRequest.WorldRay = InRequest.WorldRay;
+			SnapRequest.SnapDistance = SnappingConfig.ObjectTransform.SnapDistance;
+			SnapRequest.Transform = InRequest.Transform;
+			SnapRequest.CoordinateSpace = InRequest.RequestCoordinateSpace;
+			SnapRequest.CoordinateTransform = InRequest.CoordinateTransform;
+
+			FSceneSnapQueryResult SnapResult;
+	        SnapResult.TargetType = GetTargetTypes();
+
+			FVector SnappedPosition;
+			ESceneSnapQueryTargetResult TargetResult = SnapPosition(SnapRequest, InRequest.Position, SnappedPosition, InRequest.AxisList);
+			if (TargetResult == ESceneSnapQueryTargetResult::Snapped)
+			{
+				SnapResult.CoordinateSpace = EToolContextCoordinateSystem::World;
+				SnapResult.Position = SnappedPosition;
+				OutResults.Emplace(SnapResult);
+			}
+
+			return TargetResult;
+		}
+
+		return ESceneSnapQueryTargetResult::Disabled;
+	}
+
+	ESceneSnapQueryTargetResult FGizmoTransformSnapper::SnapPosition(const FVector& InPosition, FVector& OutSnappedPosition, const EAxisList::Type InAxisList) const
+	{
+		OutSnappedPosition = InPosition;
+
+		if (!QueriesAPI)
+		{
+			return ESceneSnapQueryTargetResult::Unsupported;
+		}
+
+		if (const FToolContextSnappingConfiguration SnappingConfig = QueriesAPI->GetCurrentSnappingSettings();
+			SnappingConfig.ObjectTransform.bEnable)
+		{
+			FSceneSnapQueryResult SnapResult;
+			SnapResult.TargetType = ESceneSnapQueryTargetType::ObjectTransform;
+
+			FGizmoTransformSnapperRequest SnapRequest;
+			SnapRequest.SnapDistance = SnappingConfig.ObjectTransform.SnapDistance;
+
+			return SnapPosition(SnapRequest, InPosition, OutSnappedPosition, InAxisList);
+		}
+
+		return ESceneSnapQueryTargetResult::Disabled;
+	}
+
+	ESceneSnapQueryTargetResult FGizmoTransformSnapper::SnapPositionAxis(const FGizmoTransformSnapperRequest& InTargetRequest, const double InAxisValue, double& OutSnappedAxisValue, const EAxisList::Type InAxis) const
+	{
+		OutSnappedAxisValue = InAxisValue;
+
+		if (InAxis == EAxisList::Screen)
+		{
+			return ESceneSnapQueryTargetResult::Unsupported;
+		}
+
+		const EAxis::Type Axis = EAxis::FromAxisList(InAxis);
+
+		FVector Position = FVector::ZeroVector;
+		Position.SetComponentForAxis(Axis, InAxisValue);
+
+		FVector SnappedPosition;
+		const ESceneSnapQueryTargetResult TargetResult = SnapPosition(InTargetRequest, Position, SnappedPosition, InAxis);
+
+		if (TargetResult == ESceneSnapQueryTargetResult::Snapped)
+		{
+			OutSnappedAxisValue = SnappedPosition.GetComponentForAxis(Axis);
+		}
+
+		return TargetResult;
+	}
+
+	ESceneSnapQueryTargetResult FGizmoTransformSnapper::SnapPositionAxis(const double InAxisValue, double& OutSnappedAxisValue, const EAxisList::Type InAxis) const
+	{
+		OutSnappedAxisValue = InAxisValue;
+
+		if (!QueriesAPI)
+		{
+			return ESceneSnapQueryTargetResult::Unsupported;
+		}
+
+		if (const FToolContextSnappingConfiguration SnappingConfig = QueriesAPI->GetCurrentSnappingSettings();
+			SnappingConfig.ObjectTransform.bEnable)
+		{
+			FSceneSnapQueryResult SnapResult;
+			SnapResult.TargetType = ESceneSnapQueryTargetType::ObjectTransform;
+
+			FViewCameraState ViewState;
+			QueriesAPI->GetCurrentViewState(ViewState);
+
+			FGizmoTransformSnapperRequest SnapRequest;
+			SnapRequest.WorldRay = FRay(ViewState.Position, ViewState.Forward(), true);
+			SnapRequest.SnapDistance = SnappingConfig.ObjectTransform.SnapDistance;
+
+			return SnapPositionAxis(SnapRequest, InAxisValue, OutSnappedAxisValue, InAxis);
+		}
+
+		return ESceneSnapQueryTargetResult::Disabled;
+	}
+
+	bool FGizmoTransformSnapper::IsQueryTypeSupported(const ESceneSnapQueryType InQueryType) const
+	{
+		return InQueryType == ESceneSnapQueryType::Position;
+	}
+}
diff --git a/Engine/Source/Editor/Experimental/EditorInteractiveToolsFramework/Private/Snapping/TransformSnapping.h b/Engine/Source/Editor/Experimental/EditorInteractiveToolsFramework/Private/Snapping/TransformSnapping.h
new file mode 100644
--- /dev/null
+++ b/Engine/Source/Editor/Experimental/EditorInteractiveToolsFramework/Private/Snapping/TransformSnapping.h
@@ -0,0 +1,55 @@
+// Copyright Epic Games, Inc. All Rights Reserved.
+
+#pragma once
+
+#include "Math/MathFwd.h"
+#include "SceneQueries/SceneSnappingManager.h"
+
+class IToolsContextQueriesAPI;
+class UGizmoViewContext;
+class UInteractiveGizmoManager;
+class UInteractiveToolsContext;
+class USceneSnappingManager;
+struct FSceneHitQueryRequest;
+struct FSceneHitQueryResult;
+struct FSceneSnapQueryRequest;
+struct FSceneSnapQueryResult;
+
+class UEditorSceneSnappingManager;
+
+namespace UE::Editor::Gizmos
+{
+	struct FGizmoTransformSnapperRequest
+	{
+		EToolContextCoordinateSystem CoordinateSpace = EToolContextCoordinateSystem::World;
+		FTransform Transform;
+		FTransform CoordinateTransform;
+		float SnapDistance = 0.0f;
+		FRay WorldRay;
+	};
+
+	class FGizmoTransformSnapper final : public TSceneSnapQueryTargetHandler<FGizmoTransformSnapperRequest>
+	{
+	public:
+		explicit FGizmoTransformSnapper(const UInteractiveToolsContext* InToolsContext);
+		virtual ~FGizmoTransformSnapper() override = default;
+
+		//~ Begin FSceneSnapPolicyTargetHandler
+		virtual ESceneSnapQueryTargetResult SnapPosition(const FSceneSnapQueryRequest& InRequest, TArray<FSceneSnapQueryResult>& OutResults) const override;
+		virtual ESceneSnapQueryTargetResult SnapPosition(const FVector& InPosition, FVector& OutSnappedPosition, const EAxisList::Type InAxisList = EAxisList::All) const override;
+		virtual ESceneSnapQueryTargetResult SnapPositionAxis(const double InAxisValue, double& OutSnappedAxisValue, const EAxisList::Type InAxis) const override;
+		virtual bool IsQueryTypeSupported(const ESceneSnapQueryType InQueryType) const override;
+		//~ End FSceneSnapPolicyTargetHandler
+
+		//~ Begin TSceneSnapQueryTargetHandler
+		virtual ESceneSnapQueryTargetResult SnapPosition(const FGizmoTransformSnapperRequest& InTargetRequest, const FVector& InPosition, FVector& OutSnappedPosition, const EAxisList::Type InAxisList = EAxisList::All) const override;
+		virtual ESceneSnapQueryTargetResult SnapPositionAxis(const FGizmoTransformSnapperRequest& InTargetRequest, const double InAxisValue, double& OutSnappedPositionAxis, const EAxisList::Type InAxis) const override;
+		//~ End TSceneSnapQueryTargetHandler
+
+		virtual ESceneSnapQueryTargetType GetTargetTypes() const override { return ESceneSnapQueryTargetType::ObjectTransform; }
+
+	private:
+		TWeakObjectPtr<UInteractiveToolManager> ToolManager = nullptr;
+		const IToolsContextQueriesAPI* QueriesAPI = nullptr;
+	};
+}
diff --git a/Engine/Source/Editor/Experimental/EditorInteractiveToolsFramework/Public/Snapping/EditorSnappingManager.h b/Engine/Source/Editor/Experimental/EditorInteractiveToolsFramework/Public/Snapping/EditorSnappingManager.h
--- a/Engine/Source/Editor/Experimental/EditorInteractiveToolsFramework/Public/Snapping/EditorSnappingManager.h
+++ b/Engine/Source/Editor/Experimental/EditorInteractiveToolsFramework/Public/Snapping/EditorSnappingManager.h
@@ -4,6 +4,7 @@
 
 #include "Math/MathFwd.h"
 #include "SceneQueries/SceneSnappingManager.h"
+#include "TickableEditorObject.h"
 
 #include "EditorSnappingManager.generated.h"
 
@@ -25,6 +26,11 @@ namespace UE::Editor::Gizmos
 	EDITORINTERACTIVETOOLSFRAMEWORK_API bool DeregisterSceneSnappingManager(const UInteractiveToolsContext* InToolsContext);
 
 	EDITORINTERACTIVETOOLSFRAMEWORK_API UEditorSceneSnappingManager* FindSceneSnappingManager(const UInteractiveToolManager* InToolManager);
+
+	namespace Private
+	{
+		class FEditorSceneSnappingManagerDebug;
+	}
 }
 
 UCLASS(MinimalAPI)
@@ -33,10 +39,13 @@ class UEditorSceneSnappingManager : public USceneSnappingManager
 	GENERATED_BODY()
 
 public:
+	UEditorSceneSnappingManager();
+
 	void Initialize(const TObjectPtr<UInteractiveToolsContext>& InToolsContext);
 	void Shutdown();
 
 	virtual bool ExecuteSceneHitQuery(const FSceneHitQueryRequest& InRequest, FSceneHitQueryResult& OutResult) const override;
+	virtual bool ExecuteSceneHitQuery(const FSceneHitQueryRequest& InRequest, TArray<FSceneHitQueryResult>& OutResults) const override;
 	virtual bool ExecuteSceneSnapQuery(const FSceneSnapQueryRequest& InRequest, TArray<FSceneSnapQueryResult>& OutResults) const override;
 
 	/** Register a snapping handler for a given snap target. */
@@ -44,12 +53,17 @@ public:
 	EDITORINTERACTIVETOOLSFRAMEWORK_API void DeregisterSnapQueryTargetHandler(const FName InHandlerName);
 
 protected:
+	bool ExecuteSceneHitQueryPrimitiveTransform(const FSceneHitQueryRequest& InRequest, TArray<FSceneHitQueryResult>& OutResults) const;
+
 	bool ExecuteSceneSnapQueryPosition(const FSceneSnapQueryRequest& InRequest, TArray<FSceneSnapQueryResult>& OutResults) const;
 	bool ExecuteSceneSnapQueryRotation(const FSceneSnapQueryRequest& InRequest, TArray<FSceneSnapQueryResult>& OutResults) const;
 	bool ExecuteSceneSnapQueryRotationAngle(const FSceneSnapQueryRequest& InRequest, TArray<FSceneSnapQueryResult>& OutResults) const;
 	bool ExecuteSceneSnapQueryScale(const FSceneSnapQueryRequest& InRequest, TArray<FSceneSnapQueryResult>& OutResults) const;
 
+protected:
 	TArray<TUniquePtr<FSceneSnapQueryTargetHandler>> QueryTargetHandlers;
 
 	const IToolsContextQueriesAPI* QueriesAPI = nullptr;
+
+	TUniquePtr<UE::Editor::Gizmos::Private::FEditorSceneSnappingManagerDebug> Debug;
 };
diff --git a/Engine/Source/Editor/UnrealEd/Private/LevelEditorViewport.cpp b/Engine/Source/Editor/UnrealEd/Private/LevelEditorViewport.cpp
--- a/Engine/Source/Editor/UnrealEd/Private/LevelEditorViewport.cpp
+++ b/Engine/Source/Editor/UnrealEd/Private/LevelEditorViewport.cpp
@@ -3347,32 +3347,42 @@ bool FLevelEditorViewportClient::InputWidgetDelta(FViewport* InViewport, EAxisLi
 					}
 				}
 
+				const bool bUsesViewportITF = ModeTools->UsesViewportITF();
+
 				// We do not want actors updated if we are holding down the middle mouse button.
 				// enable MMB for New TRS Gizmos
-				const bool bEnableMMB = ModeTools->UsesViewportITF() ? !bDraggingByHandle : false;
+				const bool bEnableMMB = bUsesViewportITF ? !bDraggingByHandle : false;
 				
 				if(!MiddleMouseButtonDown || bEnableMMB)
 				{
-					bool bSnapped = FSnappingUtils::SnapActorsToNearestActor( Drag, this );
-					bSnapped = bSnapped || FSnappingUtils::SnapDraggedActorsToNearestVertex( Drag, this );
+					// Only perform snapping for legacy gizmos
+					if (!bUsesViewportITF)
+					{
+						bool bSnapped = FSnappingUtils::SnapActorsToNearestActor( Drag, this );
+						bSnapped = bSnapped || FSnappingUtils::SnapDraggedActorsToNearestVertex( Drag, this );
 
-					// If we are only changing position, project the actors onto the world
-					const bool bOnlyTranslation = !Drag.IsZero() && Rot.IsZero() && Scale.IsZero();
+						// If we are only changing position, project the actors onto the world
+						const bool bOnlyTranslation = !Drag.IsZero() && Rot.IsZero() && Scale.IsZero();
 
-					const EAxisList::Type CurrentAxis = GetCurrentWidgetAxis();
-					const bool bSingleAxisDrag = CurrentAxis == EAxisList::X || CurrentAxis == EAxisList::Y || CurrentAxis == EAxisList::Z;
-					if (!bSnapped && !bSingleAxisDrag && GetDefault<ULevelEditorViewportSettings>()->SnapToSurface.bEnabled && bOnlyTranslation)
-					{
-						TArray<AActor*> SelectedActors;
-						for (FSelectionIterator It(GEditor->GetSelectedActorIterator()); It; ++It)
+						const EAxisList::Type CurrentAxis = GetCurrentWidgetAxis();
+						const bool bSingleAxisDrag = CurrentAxis == EAxisList::X || CurrentAxis == EAxisList::Y || CurrentAxis == EAxisList::Z;
+						if (!bSnapped && !bSingleAxisDrag && GetDefault<ULevelEditorViewportSettings>()->SnapToSurface.bEnabled && bOnlyTranslation)
 						{
-							if (AActor* Actor = Cast<AActor>(*It))
+							TArray<AActor*> SelectedActors;
+							for (FSelectionIterator It(GEditor->GetSelectedActorIterator()); It; ++It)
 							{
-								SelectedActors.Add(Actor);
+								if (AActor* Actor = Cast<AActor>(*It))
+								{
+									SelectedActors.Add(Actor);
+								}
 							}
-						}
 
-						ProjectActorsIntoWorld(SelectedActors, InViewport, Drag, Rot);
+							ProjectActorsIntoWorld(SelectedActors, InViewport, Drag, Rot);
+						}
+						else
+						{
+							ApplyDeltaToSelectedElements(FTransform(Rot, Drag, Scale));
+						}
 					}
 					else
 					{
diff --git a/Engine/Source/Editor/UnrealEd/Private/Tools/EdModeInteractiveToolsContext.cpp b/Engine/Source/Editor/UnrealEd/Private/Tools/EdModeInteractiveToolsContext.cpp
--- a/Engine/Source/Editor/UnrealEd/Private/Tools/EdModeInteractiveToolsContext.cpp
+++ b/Engine/Source/Editor/UnrealEd/Private/Tools/EdModeInteractiveToolsContext.cpp
@@ -45,6 +45,8 @@
 
 //#define ENABLE_DEBUG_PRINTING
 
+#include "SnappingUtils.h"
+
 #include UE_INLINE_GENERATED_CPP_BY_NAME(EdModeInteractiveToolsContext)
 
 DEFINE_LOG_CATEGORY_STATIC(LogInteractiveToolsContext, Log, All);
@@ -191,6 +193,10 @@ public:
 		Config.bEnableScaleGridSnapping = (GetDefault<ULevelEditorViewportSettings>()->SnapScaleEnabled != 0);
 		Config.ScaleGridSize = GEditor->GetScaleGridSize();
 		Config.bEnableAbsoluteWorldSnapping = ToolsContext->GetAbsoluteWorldSnappingEnabled();
+
+		Config.ObjectTransform.bEnable = (GetDefault<ULevelEditorViewportSettings>()->bEnableActorSnap != 0);
+		Config.ObjectTransform.SnapDistance = FSnappingUtils::GetActorSnapDistance();
+		
 		return Config;
 	}
 
diff --git a/Engine/Source/Runtime/InteractiveToolsFramework/Private/SceneQueries/SceneSnappingManager.cpp b/Engine/Source/Runtime/InteractiveToolsFramework/Private/SceneQueries/SceneSnappingManager.cpp
--- a/Engine/Source/Runtime/InteractiveToolsFramework/Private/SceneQueries/SceneSnappingManager.cpp
+++ b/Engine/Source/Runtime/InteractiveToolsFramework/Private/SceneQueries/SceneSnappingManager.cpp
@@ -90,6 +90,54 @@ void FSceneHitQueryResult::InitializeHitResult(const FSceneHitQueryRequest& From
 	// initialize .Time? Need start/end for that...
 }
 
+ESceneSnapQueryTargetResult FSceneSnapQueryTargetHandler::SnapPosition(
+	const FSceneSnapQueryRequest& InRequest,
+	TArray<FSceneSnapQueryResult>& OutResults) const
+{
+	FSceneSnapQueryResult SnapResult;
+	SnapResult.TargetType = GetTargetTypes();
+
+	const ESceneSnapQueryTargetResult TargetResult = SnapPosition(InRequest.Position, SnapResult.Position, InRequest.AxisList);
+	if (TargetResult == ESceneSnapQueryTargetResult::Snapped)
+	{
+		OutResults.Emplace(SnapResult);
+	}
+
+	return TargetResult;
+}
+
+ESceneSnapQueryTargetResult FSceneSnapQueryTargetHandler::SnapRotation(
+	const FSceneSnapQueryRequest& InRequest,
+	TArray<FSceneSnapQueryResult>& OutResults) const
+{
+	FSceneSnapQueryResult SnapResult;
+	SnapResult.TargetType = GetTargetTypes();
+
+	const ESceneSnapQueryTargetResult TargetResult = SnapRotation(InRequest.DeltaRotation, SnapResult.DeltaRotation, InRequest.AxisList);
+	if (TargetResult == ESceneSnapQueryTargetResult::Snapped)
+	{
+		OutResults.Emplace(SnapResult);
+	}
+
+	return TargetResult;
+}
+
+ESceneSnapQueryTargetResult FSceneSnapQueryTargetHandler::SnapScale(
+	const FSceneSnapQueryRequest& InRequest,
+	TArray<FSceneSnapQueryResult>& OutResults) const
+{
+	FSceneSnapQueryResult SnapResult;
+    SnapResult.TargetType = GetTargetTypes();
+
+    const ESceneSnapQueryTargetResult TargetResult = SnapScale(InRequest.Scale, SnapResult.Scale, InRequest.AxisList);
+	if (TargetResult == ESceneSnapQueryTargetResult::Snapped)
+	{
+		OutResults.Emplace(SnapResult);
+	}
+
+    return TargetResult;
+}
+
 EToolContextCoordinateSystem FSceneSnapQueryRequest::GetResultCoordinateSpace() const
 {
 	return ResultCoordinateSpace.Get(RequestCoordinateSpace);
diff --git a/Engine/Source/Runtime/InteractiveToolsFramework/Public/SceneQueries/SceneSnappingManager.h b/Engine/Source/Runtime/InteractiveToolsFramework/Public/SceneQueries/SceneSnappingManager.h
--- a/Engine/Source/Runtime/InteractiveToolsFramework/Public/SceneQueries/SceneSnappingManager.h
+++ b/Engine/Source/Runtime/InteractiveToolsFramework/Public/SceneQueries/SceneSnappingManager.h
@@ -11,12 +11,14 @@
 
 #include "SceneSnappingManager.generated.h"
 
-class ISnappingPolicy;
 class AActor;
+class ISnappingPolicy;
 class UActorComponent;
-class UPrimitiveComponent;
-class UInteractiveToolManager;
 class UInteractiveGizmoManager;
+class UInteractiveToolManager;
+class UPrimitiveComponent;
+struct FSceneSnapQueryRequest;
+struct FSceneSnapQueryResult;
 
 struct FSceneQueryVisibilityFilter
 {
@@ -30,14 +32,40 @@ struct FSceneQueryVisibilityFilter
 	INTERACTIVETOOLSFRAMEWORK_API bool IsVisible(const UPrimitiveComponent* Component) const;
 };
 
+/** Types of Hit Targets that a caller may want to run hit queries against. */
+UENUM()
+enum class ESceneHitQueryTargetType : uint8
+{
+	None = 0,
+	/** Consider any mesh vertex */
+	MeshVertex = 1,
+	/** Consider any mesh edge */
+	MeshEdge = 2,
+	/** Hit any primitive with a surface representation */
+	PrimitiveSurface = 4,
+	/** Any primitive with a transform representation  */
+	PrimitiveTransform = 8,
+	/** As defined by an FName */
+	Custom = 16,
+
+	All = MeshVertex | MeshEdge | PrimitiveSurface | PrimitiveTransform | Custom
+};
+ENUM_CLASS_FLAGS(ESceneHitQueryTargetType);
+
 
 /**
 * Configuration variables for a USceneSnappingManager hit query request.
 */
 struct FSceneHitQueryRequest
 {
-	/** scene query ray */
-	FRay3d WorldRay;
+	/** Scene query ray. */
+	FRay WorldRay;
+
+	/** What does the caller want to try to hit. */
+	ESceneHitQueryTargetType TargetTypes = ESceneHitQueryTargetType::All;
+
+	/** The name of the custom TargetTypes, if applicable. */
+	TSet<FName> CustomTargetTypes;
 
 	bool bWantHitGeometryInfo = false;
 
@@ -136,19 +164,19 @@ class FSceneSnapQueryTargetHandler
 public:
 	virtual ~FSceneSnapQueryTargetHandler() = default;
 
-	/** Snaps the given transform, potentially affecting all query types (position, rotation, scale). The given transform can be absolute or a relative delta. */
-	virtual ESceneSnapQueryTargetResult SnapTransform(const FTransform& InTransform, FVector& OutSnappedTransform) const { return ESceneSnapQueryTargetResult::Unsupported; }
-
+	virtual INTERACTIVETOOLSFRAMEWORK_API ESceneSnapQueryTargetResult SnapPosition(const FSceneSnapQueryRequest& InRequest, TArray<FSceneSnapQueryResult>& OutResults) const;
 	virtual ESceneSnapQueryTargetResult SnapPosition(const FVector& InPosition, FVector& OutSnappedPosition, const EAxisList::Type InAxisList = EAxisList::All) const { return ESceneSnapQueryTargetResult::Unsupported; }
-	virtual ESceneSnapQueryTargetResult SnapPositionAxis(const double InAxisValue, double& OutSnappedAxisValue, const EAxis::Type InAxis) const { return ESceneSnapQueryTargetResult::Unsupported; }
+	virtual ESceneSnapQueryTargetResult SnapPositionAxis(const double InAxisValue, double& OutSnappedAxisValue, const EAxisList::Type InAxis) const { return ESceneSnapQueryTargetResult::Unsupported; }
 
+	virtual INTERACTIVETOOLSFRAMEWORK_API ESceneSnapQueryTargetResult SnapRotation(const FSceneSnapQueryRequest& InRequest, TArray<FSceneSnapQueryResult>& OutResults) const;
 	virtual ESceneSnapQueryTargetResult SnapRotation(const FQuat& InRotation, FQuat& OutSnappedRotation, const EAxisList::Type InAxisList = EAxisList::All) const { return ESceneSnapQueryTargetResult::Unsupported; }
-	virtual ESceneSnapQueryTargetResult SnapRotationAxisAngle(const double& InAngle, double& OutSnappedAngle, const EAxis::Type InAxis) const { return ESceneSnapQueryTargetResult::Unsupported; }
+	virtual ESceneSnapQueryTargetResult SnapRotationAxisAngle(const double& InAngle, double& OutSnappedAngle, const EAxisList::Type InAxis) const { return ESceneSnapQueryTargetResult::Unsupported; }
 
+	virtual INTERACTIVETOOLSFRAMEWORK_API ESceneSnapQueryTargetResult SnapScale(const FSceneSnapQueryRequest& InRequest, TArray<FSceneSnapQueryResult>& OutResults) const;
 	virtual ESceneSnapQueryTargetResult SnapScale(const FVector& InScale, FVector& OutSnappedScale, const EAxisList::Type InAxisList = EAxisList::All) const { return ESceneSnapQueryTargetResult::Unsupported; }
-	virtual ESceneSnapQueryTargetResult SnapScaleAxis(const double InAxisValue, double& OutSnappedAxisValue, const EAxis::Type InAxis) const { return ESceneSnapQueryTargetResult::Unsupported; }
+	virtual ESceneSnapQueryTargetResult SnapScaleAxis(const double InAxisValue, double& OutSnappedAxisValue, const EAxisList::Type InAxis) const { return ESceneSnapQueryTargetResult::Unsupported; }
 
-	virtual ESceneSnapQueryTargetType GetTargetType() const { return ESceneSnapQueryTargetType::Custom; }
+	virtual ESceneSnapQueryTargetType GetTargetTypes() const { return ESceneSnapQueryTargetType::Custom; }
 
 	/** Returns true if the given query type can -ever- be supported, regardless of any variable state. */
 	virtual bool IsQueryTypeSupported(const ESceneSnapQueryType InQueryType) const { return false; }
@@ -179,29 +207,26 @@ template <typename TargetRequestParametersType UE_REQUIRES(std::is_class_v<Targe
 class TSceneSnapQueryTargetHandler : public FSceneSnapQueryTargetHandler
 {
 public:
-	/** Snaps the given transform, potentially affecting all query types (position, rotation, scale). The given transform can be absolute or a relative delta. */
-	virtual ESceneSnapQueryTargetResult SnapTransform(const FTransform& InTransform, FVector& OutSnappedTransform) const override { return ESceneSnapQueryTargetResult::Unsupported; }
-
-	/** Snaps the given transform, potentially affecting all query types (position, rotation, scale). The given transform can be absolute or a relative delta. */
-	virtual ESceneSnapQueryTargetResult SnapTransform(const TargetRequestParametersType& InTargetRequest, const FTransform& InTransform, FVector& OutSnappedTransform) const { return ESceneSnapQueryTargetResult::Unsupported; }
-
+	virtual ESceneSnapQueryTargetResult SnapPosition(const FSceneSnapQueryRequest& InRequest, TArray<FSceneSnapQueryResult>& OutResults) const override { return FSceneSnapQueryTargetHandler::SnapPosition(InRequest, OutResults); }
 	virtual ESceneSnapQueryTargetResult SnapPosition(const FVector& InPosition, FVector& OutSnappedPosition, const EAxisList::Type InAxisList = EAxisList::All) const override { return ESceneSnapQueryTargetResult::Unsupported; }
 	virtual ESceneSnapQueryTargetResult SnapPosition(const TargetRequestParametersType& InTargetRequest, const FVector& InPosition, FVector& OutSnappedPosition, const EAxisList::Type InAxisList = EAxisList::All) const { return ESceneSnapQueryTargetResult::Unsupported; }
 
-	virtual ESceneSnapQueryTargetResult SnapPositionAxis(const double InAxisValue, double& OutSnappedAxisValue, const EAxis::Type InAxis) const override { return ESceneSnapQueryTargetResult::Unsupported; }
-	virtual ESceneSnapQueryTargetResult SnapPositionAxis(const TargetRequestParametersType& InTargetRequest, const double InAxisValue, double& OutSnappedAxisValue, const EAxis::Type InAxis) const { return ESceneSnapQueryTargetResult::Unsupported; }
+	virtual ESceneSnapQueryTargetResult SnapPositionAxis(const double InAxisValue, double& OutSnappedAxisValue, const EAxisList::Type InAxis) const override { return ESceneSnapQueryTargetResult::Unsupported; }
+	virtual ESceneSnapQueryTargetResult SnapPositionAxis(const TargetRequestParametersType& InTargetRequest, const double InAxisValue, double& OutSnappedAxisValue, const EAxisList::Type InAxis) const { return ESceneSnapQueryTargetResult::Unsupported; }
 
+	virtual ESceneSnapQueryTargetResult SnapRotation(const FSceneSnapQueryRequest& InRequest, TArray<FSceneSnapQueryResult>& OutResults) const override { return FSceneSnapQueryTargetHandler::SnapRotation(InRequest, OutResults); }
 	virtual ESceneSnapQueryTargetResult SnapRotation(const FQuat& InRotation, FQuat& OutSnappedRotation, const EAxisList::Type InAxisList = EAxisList::All) const override { return ESceneSnapQueryTargetResult::Unsupported; }
 	virtual ESceneSnapQueryTargetResult SnapRotation(const TargetRequestParametersType& InTargetRequest, const FQuat& InRotation, FQuat& OutSnappedRotation, const EAxisList::Type InAxisList = EAxisList::All) const { return ESceneSnapQueryTargetResult::Unsupported; }
 
-	virtual ESceneSnapQueryTargetResult SnapRotationAxisAngle(const double& InAngle, double& OutSnappedAngle, const EAxis::Type InAxis) const override { return ESceneSnapQueryTargetResult::Unsupported; }
-	virtual ESceneSnapQueryTargetResult SnapRotationAxisAngle(const TargetRequestParametersType& InTargetRequest, const double& InAngle, double& OutSnappedAngle, const EAxis::Type InAxis) const { return ESceneSnapQueryTargetResult::Unsupported; }
+	virtual ESceneSnapQueryTargetResult SnapRotationAxisAngle(const double& InAngle, double& OutSnappedAngle, const EAxisList::Type InAxis) const override { return ESceneSnapQueryTargetResult::Unsupported; }
+	virtual ESceneSnapQueryTargetResult SnapRotationAxisAngle(const TargetRequestParametersType& InTargetRequest, const double& InAngle, double& OutSnappedAngle, const EAxisList::Type InAxis) const { return ESceneSnapQueryTargetResult::Unsupported; }
 
+	virtual ESceneSnapQueryTargetResult SnapScale(const FSceneSnapQueryRequest& InRequest, TArray<FSceneSnapQueryResult>& OutResults) const override { return FSceneSnapQueryTargetHandler::SnapScale(InRequest, OutResults); }
 	virtual ESceneSnapQueryTargetResult SnapScale(const FVector& InScale, FVector& OutSnappedScale, const EAxisList::Type InAxisList = EAxisList::All) const override { return ESceneSnapQueryTargetResult::Unsupported; }
 	virtual ESceneSnapQueryTargetResult SnapScale(const TargetRequestParametersType& InTargetRequest, const FVector& InScale, FVector& OutSnappedScale, const EAxisList::Type InAxisList = EAxisList::All) const { return ESceneSnapQueryTargetResult::Unsupported; }
 
-	virtual ESceneSnapQueryTargetResult SnapScaleAxis(const double InAxisValue, double& OutSnappedAxisValue, const EAxis::Type InAxis) const override { return ESceneSnapQueryTargetResult::Unsupported; }
-	virtual ESceneSnapQueryTargetResult SnapScaleAxis(const TargetRequestParametersType& InTargetRequest, const double InAxisValue, double& OutSnappedAxisValue, const EAxis::Type InAxis) const { return ESceneSnapQueryTargetResult::Unsupported; }
+	virtual ESceneSnapQueryTargetResult SnapScaleAxis(const double InAxisValue, double& OutSnappedAxisValue, const EAxisList::Type InAxis) const override { return ESceneSnapQueryTargetResult::Unsupported; }
+	virtual ESceneSnapQueryTargetResult SnapScaleAxis(const TargetRequestParametersType& InTargetRequest, const double InAxisValue, double& OutSnappedAxisValue, const EAxisList::Type InAxis) const { return ESceneSnapQueryTargetResult::Unsupported; }
 };
 
 /**
@@ -212,6 +237,9 @@ struct FSceneSnapQueryRequest
 	/** Coordinate space of the request data */
 	EToolContextCoordinateSystem RequestCoordinateSpace;
 
+	/** Coordinate transform, which may be the same as Transform but could also be a screen-space plane. */
+	FTransform CoordinateTransform = FTransform::Identity;
+
 	/** Optionally specify the desired result coordinate space, which can prevent unnecessary transforms. Defaults to the requested RequestCoordinateSpace */
 	TOptional<EToolContextCoordinateSystem> ResultCoordinateSpace;
 
@@ -233,6 +261,9 @@ struct FSceneSnapQueryRequest
 	/** Optional explicitly specified scale grid */
 	TOptional<FVector> ScaleGridSize{};
 
+	/** Snap input transform in absolute coordinates (not delta) */
+	FTransform Transform = FTransform::Identity;
+
 	/** Snap input position */
 	FVector Position = FVector::ZeroVector;
 
@@ -249,11 +280,14 @@ struct FSceneSnapQueryRequest
 	double RotationAngle = 0.0;
 
 	/** Snap input axis constraints, defaults to all (unconstrained) */
-	EAxisList::Type AxisList;
+	EAxisList::Type AxisList = EAxisList::All;
 
 	/** Snap input scale */
 	FVector Scale = FVector::OneVector;
 
+	/** Scene query ray, required by some snapping operations. */
+	FRay WorldRay;
+
 	/** Optional: components to consider invisible even if they aren't. */
 	const TArray<const UPrimitiveComponent*>* ComponentsToIgnore = nullptr;
 
@@ -277,19 +311,19 @@ struct FSceneSnapQueryResult
 	AActor* TargetActor = nullptr;
 	/** Component that owns snap target */
 	UActorComponent* TargetComponent = nullptr;
-	/** What kind of geometric element was snapped to */
+	/** What kind of geometric element(s) were snapped to */
 	ESceneSnapQueryTargetType TargetType = ESceneSnapQueryTargetType::None;
 
 	/** Snap position (may not be set depending on query types) */
-	FVector Position;
+	FVector Position = FVector::ZeroVector;
 	/** Snap normal (may not be set depending on query types) */
 	FVector Normal;
 	/** Snap rotation delta (may not be set depending on query types) */
-	FQuat   DeltaRotation;
+	FQuat   DeltaRotation = FQuat::Identity;
 	/** Snap rotation angle (may not be set depending on query types) */
-	double  RotationAngle;
+	double  RotationAngle = 0.0;
 	/** Snap scale (may not be set depending on query types) */
-	FVector Scale;
+	FVector Scale = FVector::ZeroVector;
 
 	/** Vertices of triangle that contains result (for debugging, may not be set) */
 	FVector TriVertices[3];
@@ -328,6 +362,18 @@ public:
 		return false;
 	}
 
+	/**
+	* Try to find a Hit Object in the scene that satisfies the Hit Query, returning all candidate results
+	* @param Request hit query configuration
+	* @param Results hit query results
+	* @return true if any valid hit target was found
+	* @warning implementations are not required (and may not be able) to support hit testing
+	*/
+	virtual bool ExecuteSceneHitQuery(const FSceneHitQueryRequest& Request, TArray<FSceneHitQueryResult>& ResultsOut) const
+	{
+		return false;
+	}
+
 	/**
 	* Try to find Snap Targets/Results in the scene that satisfy the Snap Query.
 	* @param Request snap query configuration
diff --git a/Engine/Source/Runtime/InteractiveToolsFramework/Public/ToolContextInterfaces.h b/Engine/Source/Runtime/InteractiveToolsFramework/Public/ToolContextInterfaces.h
--- a/Engine/Source/Runtime/InteractiveToolsFramework/Public/ToolContextInterfaces.h
+++ b/Engine/Source/Runtime/InteractiveToolsFramework/Public/ToolContextInterfaces.h
@@ -181,6 +181,16 @@ struct FToolContextSnappingConfiguration
 
 	/** Specify whether Position snapping in World Coordinate System should be Absolute or Relative */
 	bool bEnableAbsoluteWorldSnapping = false;
+
+	/** Stores parameters applicable to actor/object transform snapping, corresponds with ESceneSnapQueryTargetType::ObjectTransform */
+	struct FObjectTransformSnapping
+	{
+		/** Specify whether this snapping type is enabled */
+		bool bEnable = false;
+
+		/** The distance/radius around a candidate transform to automatically snap to. 0.0 requires precise placement. */
+		float SnapDistance = 0.0f;
+	} ObjectTransform;
 };
 
 
