[ITF Gizmos] Added view interaction state dependent visibility

diff --git a/Engine/Source/Runtime/InteractiveToolsFramework/Private/BaseGizmos/GizmoElementBase.cpp b/Engine/Source/Runtime/InteractiveToolsFramework/Private/BaseGizmos/GizmoElementBase.cpp
--- a/Engine/Source/Runtime/InteractiveToolsFramework/Private/BaseGizmos/GizmoElementBase.cpp
+++ b/Engine/Source/Runtime/InteractiveToolsFramework/Private/BaseGizmos/GizmoElementBase.cpp
@@ -382,15 +382,27 @@ bool UGizmoElementBase::GetSubElementsRecursive(TArray<TObjectPtr<UGizmoElementB
 	return !OutElements.IsEmpty();
 }
 
-bool UGizmoElementBase::IsVisible(const FSceneView* View, EGizmoElementInteractionState InCurrentInteractionState,
-	const FTransform& InLocalToWorldTransform, const FVector& InLocalCenter) const
+bool UGizmoElementBase::IsVisible(
+	const FSceneView* View,
+	EViewInteractionState InCurrentViewInteractionState,
+	EGizmoElementInteractionState InCurrentInteractionState,
+	const FTransform& InLocalToWorldTransform,
+	const FVector& InLocalCenter) const
 {
 	return (GetVisibleState() &&
-		GetEnabledForInteractionState(InCurrentInteractionState) && 
+		(!IsVisibleFunction.IsSet() || IsVisibleFunction(View, InCurrentViewInteractionState, InCurrentInteractionState, InLocalToWorldTransform, InLocalCenter)) &&
+		GetEnabledForInteractionState(InCurrentInteractionState) &&
 		GetEnabledForViewProjection(View->IsPerspectiveProjection()) &&
 		GetViewDependentVisibility(View, InLocalToWorldTransform, InLocalCenter));
 }
 
+bool UGizmoElementBase::IsVisible(const FSceneView* View, EGizmoElementInteractionState InCurrentInteractionState,
+	const FTransform& InLocalToWorldTransform, const FVector& InLocalCenter) const
+{
+	// Call with "All" view interaction state, as this (deprecated) call path is unaware of the actual state
+	return IsVisible(View, EViewInteractionState::All, InCurrentInteractionState, InLocalToWorldTransform, InLocalCenter);
+}
+
 bool UGizmoElementBase::IsHittable(const UGizmoViewContext* ViewContext, const FTransform& InLocalToWorldTransform, const FVector& InLocalCenter) const
 {
 	return (GetHittableState() &&
@@ -422,7 +434,7 @@ bool UGizmoElementBase::UpdateRenderState(IToolsContextRenderAPI* RenderAPI, con
 
 	InOutRenderState.MeshRenderState.Update(MeshRenderAttributes);
 
-	if (IsVisible(View, InOutRenderState.InteractionState, InOutRenderState.LocalToWorldTransform, InLocalCenter))
+	if (IsVisible(View, RenderAPI->GetViewInteractionState(), InOutRenderState.InteractionState, InOutRenderState.LocalToWorldTransform, InLocalCenter))
 	{
 		bool bIsPerspectiveView;
 		FVector ViewLocation, ViewDirection, ViewUp;
@@ -608,6 +620,13 @@ bool UGizmoElementBase::GetEnabledForSubduedState()
 	return bEnabledForSubduedState;
 }
 
+void UGizmoElementBase::SetIsVisibleFunction(
+	TUniqueFunction<bool(const FSceneView* InView, EViewInteractionState InCurrentViewInteractionState, EGizmoElementInteractionState InCurrentInteractionState, const FTransform&
+	InLocalToWorldTransform, const FVector& InLocalCenter)>&& InIsVisibleFunction)
+{
+	IsVisibleFunction = MoveTemp(InIsVisibleFunction);
+}
+
 void UGizmoElementBase::SetPartIdentifier(uint32 InPartIdentifier)
 {
 	PartIdentifier = InPartIdentifier;
diff --git a/Engine/Source/Runtime/InteractiveToolsFramework/Public/BaseGizmos/GizmoElementBase.h b/Engine/Source/Runtime/InteractiveToolsFramework/Public/BaseGizmos/GizmoElementBase.h
--- a/Engine/Source/Runtime/InteractiveToolsFramework/Public/BaseGizmos/GizmoElementBase.h
+++ b/Engine/Source/Runtime/InteractiveToolsFramework/Public/BaseGizmos/GizmoElementBase.h
@@ -208,6 +208,12 @@ public:
 	INTERACTIVETOOLSFRAMEWORK_API virtual void SetEnabledForSubduedState(bool bInEnabledForSubduedState);
 	INTERACTIVETOOLSFRAMEWORK_API virtual bool GetEnabledForSubduedState();
 
+	// Sets the custom visibility function. Note that this is called in addition to built-in logic rather than as a replacement.
+	INTERACTIVETOOLSFRAMEWORK_API virtual void SetIsVisibleFunction(
+		TUniqueFunction<bool(
+			const FSceneView* InView, EViewInteractionState InCurrentViewInteractionState, EGizmoElementInteractionState InCurrentInteractionState, 
+			const FTransform& InLocalToWorldTransform, const FVector& InLocalCenter)>&& InIsVisibleFunction);
+
 	// For an element hierarchy representing multiple parts of a single gizmo, the part identifier establishes 
 	// a correspondence between a gizmo part and the elements that represent that part. The recognized
 	// part identifier values should be defined in the gizmo. Gizmo part identifiers must be greater than or 
@@ -520,6 +526,11 @@ protected:
 	UPROPERTY()
 	bool bEnabledForSubduedState = true;
 
+	// An optional custom function to determine visibility of the element. Note that this is called in addition to the built-in checks, not in place of it.
+	TUniqueFunction<bool(
+		const FSceneView* InView, EViewInteractionState InCurrentViewInteractionState, EGizmoElementInteractionState InCurrentInteractionState, 
+		const FTransform& InLocalToWorldTransform, const FVector& InLocalCenter)> IsVisibleFunction = nullptr;
+
 	// Part identifier
 	UPROPERTY()
 	uint32 PartIdentifier = DefaultPartIdentifier;
@@ -624,6 +635,11 @@ protected:
 	INTERACTIVETOOLSFRAMEWORK_API bool GetSubElementsRecursive(TArray<TObjectPtr<UGizmoElementBase>>& OutElements, const uint32 InPartId) const;
 
 	// Return whether element is currently visible.
+	INTERACTIVETOOLSFRAMEWORK_API virtual bool IsVisible(const FSceneView* View, EViewInteractionState InCurrentViewInteractionState, EGizmoElementInteractionState InCurrentInteractionState, 
+		const FTransform& InLocalToWorldTransform, const FVector& InLocalCenter) const;
+
+	// Return whether element is currently visible.
+	UE_DEPRECATED(5.8, "Use IsVisible that accepts a ViewInteractionState instead.")
 	INTERACTIVETOOLSFRAMEWORK_API virtual bool IsVisible(const FSceneView* View, EGizmoElementInteractionState InCurrentInteractionState, 
 		const FTransform& InLocalToWorldTransform, const FVector& InLocalCenter) const;
 
diff --git a/Engine/Source/Runtime/InteractiveToolsFramework/Public/ToolContextInterfaces.h b/Engine/Source/Runtime/InteractiveToolsFramework/Public/ToolContextInterfaces.h
--- a/Engine/Source/Runtime/InteractiveToolsFramework/Public/ToolContextInterfaces.h
+++ b/Engine/Source/Runtime/InteractiveToolsFramework/Public/ToolContextInterfaces.h
@@ -423,7 +423,9 @@ UENUM()
 enum class EViewInteractionState {
 	None = 0,
 	Hovered = 1,
-	Focused = 2
+	Focused = 2,
+
+	All = Hovered | Focused
 };
 ENUM_CLASS_FLAGS(EViewInteractionState);
 
