[Editor Gizmos] Added Transform Adjuster support to Hit Targets

diff --git a/Engine/Source/Runtime/InteractiveToolsFramework/Private/BaseGizmos/GizmoElementHitTargets.cpp b/Engine/Source/Runtime/InteractiveToolsFramework/Private/BaseGizmos/GizmoElementHitTargets.cpp
--- a/Engine/Source/Runtime/InteractiveToolsFramework/Private/BaseGizmos/GizmoElementHitTargets.cpp
+++ b/Engine/Source/Runtime/InteractiveToolsFramework/Private/BaseGizmos/GizmoElementHitTargets.cpp
@@ -4,6 +4,7 @@
 #include "BaseGizmos/GizmoElementBase.h"
 #include "BaseGizmos/GizmoViewContext.h"
 #include "BaseGizmos/TransformProxy.h"
+#include "BaseGizmos/ViewBasedTransformAdjusters.h"
 #include "Engine/EngineTypes.h"    // FHitResult
 
 #include UE_INLINE_GENERATED_CPP_BY_NAME(GizmoElementHitTargets)
@@ -42,8 +43,14 @@ FInputRayHit UGizmoElementHitTarget::IsHit(const FInputDeviceRay& ClickPos) cons
 {
 	if (GizmoElement && GizmoViewContext && GizmoTransformProxy && (!Condition || Condition(ClickPos)))
 	{
+		FTransform Transform = GizmoTransformProxy->GetTransform();
+		if (TransformAdjuster.IsValid())
+		{
+			Transform = TransformAdjuster->GetAdjustedComponentToWorld(*GizmoViewContext, Transform);
+		}
+
 		UGizmoElementBase::FLineTraceTraversalState LineTraceState;
-		LineTraceState.Initialize(GizmoViewContext, GizmoTransformProxy->GetTransform());
+		LineTraceState.Initialize(GizmoViewContext, Transform);
 
 		UGizmoElementBase::FLineTraceOutput LineTraceOutput;		
 		return GizmoElement->LineTrace(GizmoViewContext, LineTraceState, ClickPos.WorldRay.Origin, ClickPos.WorldRay.Direction, LineTraceOutput);
@@ -91,6 +98,16 @@ void UGizmoElementHitTarget::UpdateSubdueState(bool bSubdued)
 	ToggleElementInteractionState(GizmoElement, EGizmoElementInteractionState::Subdued, bSubdued);
 }
 
+void UGizmoElementHitTarget::SetTransformAdjuster(const TSharedPtr<UE::GizmoRenderingUtil::IViewBasedTransformAdjuster>& InAdjuster)
+{
+	TransformAdjuster = InAdjuster;
+}
+
+TSharedPtr<UE::GizmoRenderingUtil::IViewBasedTransformAdjuster> UGizmoElementHitTarget::GetTransformAdjuster() const
+{
+	return TransformAdjuster;
+}
+
 UGizmoElementHitTarget* UGizmoElementHitTarget::Construct(UGizmoElementBase* InGizmoElement, UGizmoViewContext* InGizmoViewContext, UObject* Outer)
 {
 	UGizmoElementHitTarget* NewHitTarget = NewObject<UGizmoElementHitTarget>(Outer);
@@ -103,8 +120,14 @@ FInputRayHit UGizmoElementHitMultiTarget::IsHit(const FInputDeviceRay& ClickPos)
 {
 	if (GizmoElement && GizmoViewContext && GizmoTransformProxy && (!Condition || Condition(ClickPos)))
 	{
+		FTransform Transform = GizmoTransformProxy->GetTransform();
+		if (TransformAdjuster.IsValid())
+		{
+			Transform = TransformAdjuster->GetAdjustedComponentToWorld(*GizmoViewContext, Transform);
+		}
+
 		UGizmoElementBase::FLineTraceTraversalState LineTraceState;
-		LineTraceState.Initialize(GizmoViewContext, GizmoTransformProxy->GetTransform());
+		LineTraceState.Initialize(GizmoViewContext, Transform);
 
 		UGizmoElementBase::FLineTraceOutput LineTraceOutput;
 		return GizmoElement->LineTrace(GizmoViewContext, LineTraceState, ClickPos.WorldRay.Origin, ClickPos.WorldRay.Direction, LineTraceOutput);
@@ -203,6 +226,16 @@ void UGizmoElementHitMultiTarget::UpdateHittableState(bool bHittable, uint32 Par
 	GizmoElement->UpdatePartHittableState(bHittable, PartIdentifier, bInAllowMultipleElements);
 }
 
+void UGizmoElementHitMultiTarget::SetTransformAdjuster(const TSharedPtr<UE::GizmoRenderingUtil::IViewBasedTransformAdjuster>& InAdjuster)
+{
+	TransformAdjuster = InAdjuster;
+}
+
+TSharedPtr<UE::GizmoRenderingUtil::IViewBasedTransformAdjuster> UGizmoElementHitMultiTarget::GetTransformAdjuster() const
+{
+	return TransformAdjuster;
+}
+
 UGizmoElementHitMultiTarget* UGizmoElementHitMultiTarget::Construct(UGizmoElementBase* InGizmoElement, UGizmoViewContext* InGizmoViewContext, UObject* Outer)
 {
 	UGizmoElementHitMultiTarget* NewHitTarget = NewObject<UGizmoElementHitMultiTarget>(Outer);
diff --git a/Engine/Source/Runtime/InteractiveToolsFramework/Public/BaseGizmos/GizmoElementHitTargets.h b/Engine/Source/Runtime/InteractiveToolsFramework/Public/BaseGizmos/GizmoElementHitTargets.h
--- a/Engine/Source/Runtime/InteractiveToolsFramework/Public/BaseGizmos/GizmoElementHitTargets.h
+++ b/Engine/Source/Runtime/InteractiveToolsFramework/Public/BaseGizmos/GizmoElementHitTargets.h
@@ -9,6 +9,11 @@
 #include "BaseGizmos/TransformProxy.h"
 #include "GizmoElementHitTargets.generated.h"
 
+namespace UE::GizmoRenderingUtil
+{
+	class IViewBasedTransformAdjuster;
+}
+
 class UGizmoElementBase;
 
 /**
@@ -54,11 +59,21 @@ public:
 
 	INTERACTIVETOOLSFRAMEWORK_API virtual void UpdateSubdueState(bool bSubdued) override;
 
+	INTERACTIVETOOLSFRAMEWORK_API virtual void SetTransformAdjuster(const TSharedPtr<UE::GizmoRenderingUtil::IViewBasedTransformAdjuster>& InAdjuster) override;
+	INTERACTIVETOOLSFRAMEWORK_API virtual TSharedPtr<UE::GizmoRenderingUtil::IViewBasedTransformAdjuster> GetTransformAdjuster() const override;
+
 public:
 	static INTERACTIVETOOLSFRAMEWORK_API UGizmoElementHitTarget* Construct(
 		UGizmoElementBase* InGizmoElement,
 		UGizmoViewContext* InGizmoViewContext,
 		UObject* Outer = (UObject*)GetTransientPackage());
+
+protected:
+	/**
+	 * An optional transform adjuster to apply when testing hit geometry.
+	 * Note that while the name indicates it's view-based, it may also be used to apply any transform.
+	 */
+	TSharedPtr<UE::GizmoRenderingUtil::IViewBasedTransformAdjuster> TransformAdjuster = nullptr;
 };
 
 /**
@@ -120,10 +135,19 @@ public:
 
 	INTERACTIVETOOLSFRAMEWORK_API virtual void UpdateHittableState(bool bHittable, uint32 PartIdentifier, bool bInAllowMultipleElements);
 
+	INTERACTIVETOOLSFRAMEWORK_API virtual void SetTransformAdjuster(const TSharedPtr<UE::GizmoRenderingUtil::IViewBasedTransformAdjuster>& InAdjuster) override;
+	INTERACTIVETOOLSFRAMEWORK_API virtual TSharedPtr<UE::GizmoRenderingUtil::IViewBasedTransformAdjuster> GetTransformAdjuster() const override;
+
 public:
 	static INTERACTIVETOOLSFRAMEWORK_API UGizmoElementHitMultiTarget* Construct(
 		UGizmoElementBase* InGizmoElement,
 		UGizmoViewContext* InGizmoViewContext,
 		UObject* Outer = (UObject*)GetTransientPackage());
 
+protected:
+	/**
+	 * An optional transform adjuster to apply when testing hit geometry.
+	 * Note that while the name indicates it's view-based, it may also be used to apply any transform.
+	 */
+	TSharedPtr<UE::GizmoRenderingUtil::IViewBasedTransformAdjuster> TransformAdjuster = nullptr;
 };
diff --git a/Engine/Source/Runtime/InteractiveToolsFramework/Public/BaseGizmos/GizmoInterfaces.h b/Engine/Source/Runtime/InteractiveToolsFramework/Public/BaseGizmos/GizmoInterfaces.h
--- a/Engine/Source/Runtime/InteractiveToolsFramework/Public/BaseGizmos/GizmoInterfaces.h
+++ b/Engine/Source/Runtime/InteractiveToolsFramework/Public/BaseGizmos/GizmoInterfaces.h
@@ -14,6 +14,11 @@
 
 #include "GizmoInterfaces.generated.h"
 
+namespace UE::GizmoRenderingUtil
+{
+	class IViewBasedTransformAdjuster;
+}
+
 class IToolsContextRenderAPI;
 class UObject;
 struct FFrame;
@@ -141,6 +146,12 @@ public:
 	 */
 	UFUNCTION()
 	virtual void UpdateSubdueState(bool bSubdued) = 0;
+
+	/** Sets the optional transform adjuster to apply to the tested geometry. */
+	virtual void SetTransformAdjuster(const TSharedPtr<UE::GizmoRenderingUtil::IViewBasedTransformAdjuster>& InAdjuster) { /** Empty default implementation. */ }
+
+	/** Get the (optional) transform adjuster applied to this Target. */
+	virtual TSharedPtr<UE::GizmoRenderingUtil::IViewBasedTransformAdjuster> GetTransformAdjuster() const { return nullptr; /** No adjuster by default. */ }
 };
 
 UINTERFACE(MinimalAPI)
@@ -201,6 +212,12 @@ public:
 	 */
 	UFUNCTION()
 	virtual void UpdateHittableState(bool bHittable, uint32 InPartIdentifier) = 0;
+
+	/** Sets the optional transform adjuster to apply to the tested geometry. */
+	virtual void SetTransformAdjuster(const TSharedPtr<UE::GizmoRenderingUtil::IViewBasedTransformAdjuster>& InAdjuster) { /** Empty default implementation. */ }
+
+	/** Get the (optional) transform adjuster applied to this Target. */
+	virtual TSharedPtr<UE::GizmoRenderingUtil::IViewBasedTransformAdjuster> GetTransformAdjuster() const { return nullptr; /** No adjuster by default. */ }
 };
 
 
