[ITF Gizmos] Add support for orthographic flattening were possible & reasonable in ITF gizmos.

diff --git a/Engine/Source/Editor/Experimental/EditorInteractiveToolsFramework/Private/EditorGizmos/TransformGizmo.cpp b/Engine/Source/Editor/Experimental/EditorInteractiveToolsFramework/Private/EditorGizmos/TransformGizmo.cpp
--- a/Engine/Source/Editor/Experimental/EditorInteractiveToolsFramework/Private/EditorGizmos/TransformGizmo.cpp
+++ b/Engine/Source/Editor/Experimental/EditorInteractiveToolsFramework/Private/EditorGizmos/TransformGizmo.cpp
@@ -485,9 +485,10 @@ void UTransformGizmo::Render(IToolsContextRenderAPI* RenderAPI)
 
 		if (IsVisible())
 		{
+			const UE::GizmoRenderingUtil::FSceneViewWrapper WrappedSceneView(*RenderAPI->GetSceneView());
+		
 			UGizmoElementBase::FRenderTraversalState RenderState;
-			RenderState.Initialize(RenderAPI->GetSceneView(), GetGizmoTransform(), RenderAPI->GetCameraState().DPIScale);
-			GizmoElementRoot->Render(RenderAPI, RenderState);
+			RenderState.Initialize(WrappedSceneView, GetGizmoTransform(), RenderAPI->GetCameraState().DPIScale);
 		}
 
 
diff --git a/Engine/Source/Runtime/InteractiveToolsFramework/Private/BaseGizmos/GizmoElementBase.cpp b/Engine/Source/Runtime/InteractiveToolsFramework/Private/BaseGizmos/GizmoElementBase.cpp
--- a/Engine/Source/Runtime/InteractiveToolsFramework/Private/BaseGizmos/GizmoElementBase.cpp
+++ b/Engine/Source/Runtime/InteractiveToolsFramework/Private/BaseGizmos/GizmoElementBase.cpp
@@ -21,6 +21,12 @@ namespace GizmoElementBaseLocals
 		OutViewDirection = View->GetViewDirection();
 		OutViewUp = View->GetViewUp();
 	}
+	
+	static TAutoConsoleVariable<bool> CVarAllowFlattenedRendering(
+		TEXT("ITF.Gizmos.AllowFlattenedRendering"),
+		true,
+		TEXT("When enabled, gizmo rendering will be flattened when possible, improving display at very large orthographic zoom levels"),
+		ECVF_RenderThreadSafe);
 }
 
 bool UGizmoElementBase::GetViewDependentVisibility(const FSceneView* View, const FTransform& InLocalToWorldTransform, const FVector& InLocalCenter) const
@@ -504,6 +510,40 @@ bool UGizmoElementBase::GetVisibleState() const
 	return (static_cast<uint8>(ElementState) & static_cast<uint8>(EGizmoElementState::Visible));
 }
 
+void UGizmoElementBase::FRenderTraversalState::Initialize(const UE::GizmoRenderingUtil::ISceneViewInterface& InView, const FTransform& InTransform, const double InDPIScale)
+{
+	LocalToWorldTransform = InTransform;
+	PixelToWorldScale = UE::GizmoRenderingUtil::CalculateViewDependentScaleAndFlatten(&InView, InTransform.GetLocation(), 0.5f, FlattenScale) * InDPIScale;
+}
+
+PRAGMA_DISABLE_DEPRECATION_WARNINGS
+void UGizmoElementBase::FRenderTraversalState::Initialize(const FSceneView* InSceneView, FTransform InTransform, const double InDPIScale)
+{
+	UE::GizmoRenderingUtil::FSceneViewWrapper Wrapper(*InSceneView);
+	Initialize(Wrapper, InTransform, InDPIScale);
+}
+PRAGMA_ENABLE_DEPRECATION_WARNINGS
+
+FScaleMatrix UGizmoElementBase::FRenderTraversalState::SafeFlattenMatrix() const
+{
+	if (GizmoElementBaseLocals::CVarAllowFlattenedRendering.GetValueOnAnyThread() && LocalToWorldTransform.GetRotation().IsIdentity())
+	{
+		return FScaleMatrix(FlattenScale);
+	}
+	return FScaleMatrix(FVector::One());
+}
+
+void UGizmoElementBase::FLineTraceTraversalState::Initialize(const UGizmoViewContext* InGizmoViewContext, FTransform InTransform)
+{
+	check(InGizmoViewContext);
+
+	LocalToWorldTransform = InTransform;
+	bIsPerspectiveProjection = InGizmoViewContext->IsPerspectiveProjection();
+
+	constexpr bool bUseDPIScaleFromViewContext = true;
+	PixelToWorldScale = UE::GizmoRenderingUtil::CalculateLocalPixelToWorldScale(InGizmoViewContext, InTransform.GetLocation(), bUseDPIScaleFromViewContext);
+}
+
 void UGizmoElementBase::DrawHUD(FCanvas* Canvas, IToolsContextRenderAPI* RenderAPI, const FRenderTraversalState& RenderState)
 {
 	// Empty default implementation
diff --git a/Engine/Source/Runtime/InteractiveToolsFramework/Private/BaseGizmos/GizmoElementBox.cpp b/Engine/Source/Runtime/InteractiveToolsFramework/Private/BaseGizmos/GizmoElementBox.cpp
--- a/Engine/Source/Runtime/InteractiveToolsFramework/Private/BaseGizmos/GizmoElementBox.cpp
+++ b/Engine/Source/Runtime/InteractiveToolsFramework/Private/BaseGizmos/GizmoElementBox.cpp
@@ -22,7 +22,7 @@ void UGizmoElementBox::Render(IToolsContextRenderAPI* RenderAPI, const FRenderTr
 			FTransform RenderLocalToWorldTransform = FTransform(LocalRotation) * CurrentRenderState.LocalToWorldTransform;
 			const FVector HalfDimensions = Dimensions * 0.5;
 			FPrimitiveDrawInterface* PDI = RenderAPI->GetPrimitiveDrawInterface();
-			DrawBox(PDI, RenderLocalToWorldTransform.ToMatrixWithScale(), HalfDimensions, UseMaterial->GetRenderProxy(), SDPG_Foreground);
+			DrawBox(PDI, RenderLocalToWorldTransform.ToMatrixWithScale() * CurrentRenderState.SafeFlattenMatrix(), HalfDimensions, UseMaterial->GetRenderProxy(), SDPG_Foreground);
 		}
 	}
 }
diff --git a/Engine/Source/Runtime/InteractiveToolsFramework/Private/BaseGizmos/GizmoElementCone.cpp b/Engine/Source/Runtime/InteractiveToolsFramework/Private/BaseGizmos/GizmoElementCone.cpp
--- a/Engine/Source/Runtime/InteractiveToolsFramework/Private/BaseGizmos/GizmoElementCone.cpp
+++ b/Engine/Source/Runtime/InteractiveToolsFramework/Private/BaseGizmos/GizmoElementCone.cpp
@@ -80,7 +80,7 @@ void UGizmoElementCone::Render(IToolsContextRenderAPI* RenderAPI, const FRenderT
 			const float HeightOverConeSide = Height / ConeSide;
 			const float Angle = FMath::Acos(HeightOverConeSide);
 			FPrimitiveDrawInterface* PDI = RenderAPI->GetPrimitiveDrawInterface();
-			DrawCone(PDI, RenderLocalToWorldTransform.ToMatrixWithScale(), Angle, Angle, NumSides, false, FColor::White, UseMaterial->GetRenderProxy(), SDPG_Foreground);
+			DrawCone(PDI, RenderLocalToWorldTransform.ToMatrixWithScale() * CurrentRenderState.SafeFlattenMatrix(), Angle, Angle, NumSides, false, FColor::White, UseMaterial->GetRenderProxy(), SDPG_Foreground);
 
 			if (bEndCaps)
 			{
diff --git a/Engine/Source/Runtime/InteractiveToolsFramework/Private/BaseGizmos/GizmoElementCylinder.cpp b/Engine/Source/Runtime/InteractiveToolsFramework/Private/BaseGizmos/GizmoElementCylinder.cpp
--- a/Engine/Source/Runtime/InteractiveToolsFramework/Private/BaseGizmos/GizmoElementCylinder.cpp
+++ b/Engine/Source/Runtime/InteractiveToolsFramework/Private/BaseGizmos/GizmoElementCylinder.cpp
@@ -159,6 +159,7 @@ void UGizmoElementCylinder::Render(IToolsContextRenderAPI* RenderAPI, const FRen
 			const FVector OriginOffset = Direction * HalfHeight;
 
 			const FTransform RenderLocalToWorldTransform = FTransform(Rotation, OriginOffset) * CurrentRenderState.LocalToWorldTransform;
+			const FMatrix RenderLocalToWorldMatrix = RenderLocalToWorldTransform.ToMatrixWithScale() * CurrentRenderState.SafeFlattenMatrix();
 			FPrimitiveDrawInterface* PDI = RenderAPI->GetPrimitiveDrawInterface();
 
 			const FColor Color = CurrentRenderState.GetCurrentVertexColor().ToFColor(true);
@@ -313,7 +314,7 @@ void UGizmoElementCylinder::Render(IToolsContextRenderAPI* RenderAPI, const FRen
 				MeshBuilder.AddVertices(Vertices);
 				MeshBuilder.AddTriangles(TriangleIndices);
 
-				MeshBuilder.Draw(PDI, RenderLocalToWorldTransform.ToMatrixWithScale(), UseMaterial->GetRenderProxy(), SDPG_Foreground);
+				MeshBuilder.Draw(PDI, RenderLocalToWorldMatrix, UseMaterial->GetRenderProxy(), SDPG_Foreground);
 			};
 
 			if (bIsDashed)
@@ -362,7 +363,7 @@ void UGizmoElementCylinder::Render(IToolsContextRenderAPI* RenderAPI, const FRen
 
 					GizmoElementCylinderLocals::DrawCylinder(
 						PDI,
-						RenderLocalToWorldTransform.ToMatrixWithScale(),
+						RenderLocalToWorldMatrix,
 						(-FVector::UpVector * HalfHeight) + FVector::UpVector * HalfHeight,
 						FVector::XAxisVector,
 						FVector::YAxisVector,
@@ -379,7 +380,7 @@ void UGizmoElementCylinder::Render(IToolsContextRenderAPI* RenderAPI, const FRen
 				{
 					GizmoElementCylinderLocals::DrawCylinder(
 						PDI,
-						RenderLocalToWorldTransform.ToMatrixWithScale(),
+						RenderLocalToWorldMatrix,
 						(-FVector::UpVector * HalfHeight) + FVector::UpVector * HalfHeight,
 						FVector::XAxisVector,
 						FVector::YAxisVector,
diff --git a/Engine/Source/Runtime/InteractiveToolsFramework/Private/BaseGizmos/GizmoElementTorus.cpp b/Engine/Source/Runtime/InteractiveToolsFramework/Private/BaseGizmos/GizmoElementTorus.cpp
--- a/Engine/Source/Runtime/InteractiveToolsFramework/Private/BaseGizmos/GizmoElementTorus.cpp
+++ b/Engine/Source/Runtime/InteractiveToolsFramework/Private/BaseGizmos/GizmoElementTorus.cpp
@@ -40,7 +40,7 @@ void UGizmoElementTorus::Render(IToolsContextRenderAPI* RenderAPI, const FRender
 
 			const FColor Color = CurrentRenderState.GetCurrentVertexColor().ToFColor(true);
 
-			DrawTorus(PDI, CurrentRenderState.LocalToWorldTransform.ToMatrixWithScale(),
+			DrawTorus(PDI, CurrentRenderState.LocalToWorldTransform.ToMatrixWithScale() * CurrentRenderState.SafeFlattenMatrix(),
 				BeginAxis, TorusSideAxis, Color, GetCurrentRadius(), GetCurrentInnerRadius(), NumSegments, NumInnerSlices,  
 				UseMaterial->GetRenderProxy(), SDPG_Foreground, bPartial, PartialAngle, bEndCaps);
 		}
diff --git a/Engine/Source/Runtime/InteractiveToolsFramework/Public/BaseGizmos/GizmoElementBase.h b/Engine/Source/Runtime/InteractiveToolsFramework/Public/BaseGizmos/GizmoElementBase.h
--- a/Engine/Source/Runtime/InteractiveToolsFramework/Public/BaseGizmos/GizmoElementBase.h
+++ b/Engine/Source/Runtime/InteractiveToolsFramework/Public/BaseGizmos/GizmoElementBase.h
@@ -53,6 +53,10 @@ public:
 		// LocalToWorld transform 
 		// Note: non-uniform scale is not supported and the X scale element will be used for uniform scaling.
 		FTransform LocalToWorldTransform;
+		
+		// The potentially flattened scale of this state, if this view is orthographic. Otherwise, a unit vector.
+		// When rendering gizmos, `FScaleMatrix(FlattenScale)` can be applied to a rendering matrix so that gizmos never clip into the near orthographic view plane, even at extreme distances.
+		FVector FlattenScale = FVector::OneVector;
 
 		// Pixel to world scale. Also accounts for DPI scale, @see IToolsContextRenderAPI::ViewCameraState.
 		double PixelToWorldScale = 1.0;
@@ -66,12 +70,12 @@ public:
 		// Current state used for rendering lines
 		FGizmoElementLineRenderStateAttributes LineRenderState;
 
-		// Initialize state 
-		void Initialize(const FSceneView* InSceneView, FTransform InTransform, const double InDPIScale = 1.0)
-		{
-			LocalToWorldTransform = InTransform;
-			PixelToWorldScale = UE::GizmoRenderingUtil::CalculateLocalPixelToWorldScale(InSceneView, InTransform.GetLocation(), InDPIScale);
-		}
+		// Initialize state
+		INTERACTIVETOOLSFRAMEWORK_API void Initialize(const UE::GizmoRenderingUtil::ISceneViewInterface& InView, const FTransform& InTransform, const double InDPIScale = 1.0);
+
+		// Initialize state
+		UE_DEPRECATED(5.8, "Use Initialize with ISceneViewInterface instead.")
+		INTERACTIVETOOLSFRAMEWORK_API void Initialize(const FSceneView* InSceneView, FTransform InTransform, const double InDPIScale = 1.0);
 
 		// Returns the mesh material based on the current interaction state.
 		const UMaterialInterface* GetCurrentMaterial()
@@ -79,7 +83,6 @@ public:
 			return MeshRenderState.GetMaterial(InteractionState);
 		}
 
-		// Returns the mesh material based on the current interaction state.
 		const UMaterialInterface* GetCurrentMaterial() const
 		{
 			return MeshRenderState.GetMaterial(InteractionState);
@@ -108,6 +111,9 @@ public:
 		{
 			return LineRenderState.GetLineColor(InteractionState);
 		}
+		
+		// Returns the flatten scale if the LocalToWorldTransform is safe to flatten. Otherwise, returns an unscaled matrixs
+		INTERACTIVETOOLSFRAMEWORK_API FScaleMatrix SafeFlattenMatrix() const;
 	};
 
 	struct FLineTraceTraversalState
@@ -119,23 +125,14 @@ public:
 		// Pixel to world scale
 		double PixelToWorldScale = 1.0;
 
-		// View context is perspective projecion
+		// View context is perspective projection
 		bool bIsPerspectiveProjection = true;
 
 		// The criteria that determines the singular element when multiple are hit.
 		EGizmoElementHitSortType HitSortType = EGizmoElementHitSortType::Closest;
 
-		// Initialize state 
-		void Initialize(const UGizmoViewContext* InGizmoViewContext, FTransform InTransform)
-		{
-			check(InGizmoViewContext);
-
-			LocalToWorldTransform = InTransform;
-			bIsPerspectiveProjection = InGizmoViewContext->IsPerspectiveProjection();
-
-			constexpr bool bUseDPIScaleFromViewContext = true;
-			PixelToWorldScale = UE::GizmoRenderingUtil::CalculateLocalPixelToWorldScale(InGizmoViewContext, InTransform.GetLocation(), bUseDPIScaleFromViewContext);
-		}
+		// Initialize state
+		INTERACTIVETOOLSFRAMEWORK_API void Initialize(const UGizmoViewContext* InGizmoViewContext, FTransform InTransform);
 	};
 
 	/** Contains information about a line trace, in addition to FInputRayHit. */
