[Editor TRS] Axis handle is easier to select when viewing front-on #jira UE-356716
#Synced-CL 48638181

diff --git a/Engine/Source/Runtime/InteractiveToolsFramework/Private/BaseGizmos/GizmoElementArrow.cpp b/Engine/Source/Runtime/InteractiveToolsFramework/Private/BaseGizmos/GizmoElementArrow.cpp
--- a/Engine/Source/Runtime/InteractiveToolsFramework/Private/BaseGizmos/GizmoElementArrow.cpp
+++ b/Engine/Source/Runtime/InteractiveToolsFramework/Private/BaseGizmos/GizmoElementArrow.cpp
@@ -113,10 +113,19 @@ FInputRayHit UGizmoElementArrow::LineTrace(const UGizmoViewContext* ViewContext,
 	check(CylinderElement);
 	FInputRayHit Hit = CylinderElement->LineTrace(ViewContext, CurrentLineTraceState, RayOrigin, RayDirection, OutLineTraceOutput);
 
-	if (!Hit.bHit)
+	// If we didn't hit the cylinder, OR it was hit, but not a surface hit, try the arrow head
+	if (!Hit.bHit || !OutLineTraceOutput.bIsSurfaceHit)
 	{
 		check(HeadElement);
-		Hit = HeadElement->LineTrace(ViewContext, CurrentLineTraceState, RayOrigin, RayDirection, OutLineTraceOutput);
+		FLineTraceOutput HeadLineTraceOutput = OutLineTraceOutput;
+		FInputRayHit HeadHit = HeadElement->LineTrace(ViewContext, CurrentLineTraceState, RayOrigin, RayDirection, HeadLineTraceOutput);
+		if (HeadHit.bHit && HeadLineTraceOutput.bIsSurfaceHit)
+		{
+			// If we're here, the cylinder was hit, but not on it's surface - but the head was hit directly on the surface.
+			// This means we'll take priority over another element with a higher hit priority
+			OutLineTraceOutput = HeadLineTraceOutput;
+			Hit = HeadHit;
+		}
 	}
 
 	if (Hit.bHit)
