[ITF Gizmos] Fix an issue where `UGizmoElementBase::ApplyViewDepthOffset()` incorrectly positioned and scaled the objects in orthographic views.

#rnx
#Synced-CL 48950337

diff --git a/Engine/Source/Runtime/InteractiveToolsFramework/Private/BaseGizmos/GizmoElementBase.cpp b/Engine/Source/Runtime/InteractiveToolsFramework/Private/BaseGizmos/GizmoElementBase.cpp
--- a/Engine/Source/Runtime/InteractiveToolsFramework/Private/BaseGizmos/GizmoElementBase.cpp
+++ b/Engine/Source/Runtime/InteractiveToolsFramework/Private/BaseGizmos/GizmoElementBase.cpp
@@ -175,57 +175,43 @@ void UGizmoElementBase::ApplyViewDepthOffset(const FVector& InViewLocation, cons
 
 void UGizmoElementBase::ApplyViewDepthOffset(const FVector& InViewLocation, const FVector& InViewDirection, const bool bInIsPerspectiveView, const double& InViewDepthOffset, FTransform& InOutTransform) const
 {
-	FVector ViewLocation = InViewLocation;
-	const FVector ObjectLocation = InOutTransform.GetLocation();
-
-	FVector CameraToObject = ObjectLocation - ViewLocation;
-	double DistanceToObject = FVector::DotProduct(CameraToObject, InViewDirection);
-
-	const double ScaledViewDepthOffset = InViewDepthOffset * InOutTransform.GetScale3D().X;
-
-	// Orthographic view
-	if (!bInIsPerspectiveView)
-	{
-		// Back-project the effective view location based solely on the distance to the object along the view direction
-		// So that in screen space they'd overlap exactly
-		ViewLocation = ObjectLocation - InViewDirection * DistanceToObject;
-	}
-	else
+	if (bInIsPerspectiveView)
 	{
-		DistanceToObject = CameraToObject.Length();
-	}
+		// Perspective views need to move the object away from the camera while also scaling so that they have the correct size.
+		FVector CameraToObject = InOutTransform.GetLocation() - InViewLocation;
+		const double DistanceToObject = CameraToObject.Length();
+		
+		// Clamp to camera near plane
+		constexpr double MinDistance = UE_DOUBLE_KINDA_SMALL_NUMBER;
+		if (DistanceToObject < MinDistance)
+		{
+			// Object is at or behind the camera - move along the forward vector without accounting for scale
+			InOutTransform.SetLocation(InViewLocation + InViewDirection * InViewDepthOffset * InOutTransform.GetScale3D().X);
+			return;
+		}
+		
+		double TargetDistance = DistanceToObject + InViewDepthOffset;
+		if (TargetDistance <= MinDistance)
+		{
+			// Prevent negative/zero
+			TargetDistance = MinDistance;
+		}
+		
+		// Calculate scale factor to maintain constant visual size
+		const double ScaleFactor = TargetDistance / DistanceToObject;
 
-	// Clamp to camera near plane
-	constexpr double MinDistance = UE_DOUBLE_KINDA_SMALL_NUMBER;
-	if (DistanceToObject < MinDistance)
-	{
-		// Object is at or behind the camera - move along the forward vector without accounting for scale
-		InOutTransform.SetLocation(ViewLocation + InViewDirection * ScaledViewDepthOffset);
-		return;
-	}
+		// Calculate the new position that maintains the same apparent position in view
+		CameraToObject.Normalize(); // Normalize the vector to get direction
 
-	double TargetDistance = DistanceToObject + InViewDepthOffset;
-	if (TargetDistance <= UE_KINDA_SMALL_NUMBER)
-	{
-		// Prevent negative/zero
-		TargetDistance = UE_KINDA_SMALL_NUMBER;
+		// Update the transform with new location and scale
+		InOutTransform.SetLocation(InViewLocation + CameraToObject * TargetDistance);
+		InOutTransform.SetScale3D(InOutTransform.GetScale3D() * ScaleFactor);
 	}
-
-	if (TargetDistance <= MinDistance)
+	else
 	{
-		// Prevent negative/zero
-		TargetDistance = 1;
+		// Orthographic cameras can simply shift the location along the camera forward axis without any scaling applied
+		InOutTransform.SetLocation(InOutTransform.GetLocation() + InViewDirection * InViewDepthOffset);
 	}
-
-	// Calculate scale factor to maintain constant visual size
-	const double ScaleFactor = TargetDistance / DistanceToObject;
-
-	// Calculate the new position that maintains the same apparent position in view
-	CameraToObject.Normalize(); // Normalize the vector to get direction
-
-	// Update the transform with new location and scale
-	InOutTransform.SetLocation(ViewLocation + CameraToObject * TargetDistance);
-	InOutTransform.SetScale3D(InOutTransform.GetScale3D() * ScaleFactor);
 }
 
 bool UGizmoElementBase::VerifyUniformScale(const FVector& Scale) const
