[ITF] Allow construction of GizmoViewContext from view matrices and associated parameters, rather than just FSceneView (allows testability)

#Synced-CL 48895345

diff --git a/Engine/Source/Runtime/InteractiveToolsFramework/Public/BaseGizmos/GizmoViewContext.h b/Engine/Source/Runtime/InteractiveToolsFramework/Public/BaseGizmos/GizmoViewContext.h
--- a/Engine/Source/Runtime/InteractiveToolsFramework/Public/BaseGizmos/GizmoViewContext.h
+++ b/Engine/Source/Runtime/InteractiveToolsFramework/Public/BaseGizmos/GizmoViewContext.h
@@ -26,13 +26,18 @@ public:
 	class FMatrices
 	{
 	public:
+		void ResetFromViewMatrices(const FViewMatrices& InViewMatrices)
+		{
+			ProjectionMatrix = InViewMatrices.GetProjectionMatrix();
+			ViewMatrix = InViewMatrices.GetViewMatrix();
+			InvViewMatrix = InViewMatrices.GetInvViewMatrix();
+			ViewProjectionMatrix = InViewMatrices.GetViewProjectionMatrix();
+			InvViewProjectionMatrix = InViewMatrices.GetInvViewProjectionMatrix();
+		}
+
 		void ResetFromSceneView(const FSceneView& SceneView)
 		{
-			ProjectionMatrix = SceneView.ViewMatrices.GetProjectionMatrix();
-			ViewMatrix = SceneView.ViewMatrices.GetViewMatrix();
-			InvViewMatrix = SceneView.ViewMatrices.GetInvViewMatrix();
-			ViewProjectionMatrix = SceneView.ViewMatrices.GetViewProjectionMatrix();
-			InvViewProjectionMatrix = SceneView.ViewMatrices.GetInvViewProjectionMatrix();
+			ResetFromViewMatrices(SceneView.ViewMatrices);
 		}
 
 		const FMatrix& GetProjectionMatrix() const { return ProjectionMatrix; }
@@ -49,13 +54,18 @@ public:
 		FMatrix InvViewProjectionMatrix = FMatrix::Identity;
 	};
 
+	void Reset(const FIntRect& InUnscaledViewRect, const FViewMatrices& InViewMatrices, const bool bInIsPerspectiveProjection, const FVector& InViewLocation)
+	{
+		UnscaledViewRect = InUnscaledViewRect;
+		ViewMatrices.ResetFromViewMatrices(InViewMatrices);
+		bIsPerspectiveProjection = bInIsPerspectiveProjection;
+		ViewLocation = InViewLocation;
+	}
+
 	// Use this to reinitialize the object each frame for the hovered viewport.
 	void ResetFromSceneView(const FSceneView& SceneView)
 	{
-		UnscaledViewRect = SceneView.UnscaledViewRect;
-		ViewMatrices.ResetFromSceneView(SceneView);
-		bIsPerspectiveProjection = SceneView.IsPerspectiveProjection();
-		ViewLocation = SceneView.ViewLocation;
+		Reset(SceneView.UnscaledViewRect, SceneView.ViewMatrices, SceneView.IsPerspectiveProjection(), SceneView.ViewLocation);
 	}
 
 	// ISceneViewInterface
