Moved ToolGroups into ITF proper

#Synced-CL 50713120

diff --git a/Engine/Source/Runtime/InteractiveToolsFramework/Private/InteractiveToolsContext.cpp b/Engine/Source/Runtime/InteractiveToolsFramework/Private/InteractiveToolsContext.cpp
--- a/Engine/Source/Runtime/InteractiveToolsFramework/Private/InteractiveToolsContext.cpp
+++ b/Engine/Source/Runtime/InteractiveToolsFramework/Private/InteractiveToolsContext.cpp
@@ -4,6 +4,7 @@
 #include "InteractiveToolsContext.h"
 #include "ToolTargetManager.h"
 #include "ContextObjectStore.h"
+#include "ToolGroupManager.h"
 
 #include UE_INLINE_GENERATED_CPP_BY_NAME(InteractiveToolsContext)
 
@@ -13,6 +14,7 @@ UInteractiveToolsContext::UInteractiveToolsContext()
 	ToolManager = nullptr;
 	GizmoManager = nullptr;
 	TargetManager = nullptr;
+	ToolGroupManager = nullptr;
 	ContextObjectStore = nullptr;
 	ToolManagerClass = UInteractiveToolManager::StaticClass();
 
@@ -70,6 +72,17 @@ UInteractiveToolsContext::UInteractiveToolsContext()
 	{
 		ContextStoreIn->Shutdown();
 	});
+
+	SetCreateToolGroupManagerFunc(  [](const FContextInitInfo& ContextInfo)
+	{
+		UToolGroupManager* NewToolGroupManager = NewObject<UToolGroupManager>(ContextInfo.ToolsContext);
+		return NewToolGroupManager;
+	});
+
+	SetShutdownToolGroupManagerFunc( [] (UToolGroupManager* ToolGroupManagerIn)
+	{
+		ToolGroupManagerIn->Shutdown();	
+	});
 }
 
 
@@ -94,6 +107,11 @@ void UInteractiveToolsContext::SetCreateContextStoreFunc(TUniqueFunction<UContex
 	CreateContextStoreFunc = MoveTemp(Func);
 }
 
+void UInteractiveToolsContext::SetCreateToolGroupManagerFunc(TUniqueFunction<UToolGroupManager*(const FContextInitInfo&)> Func)
+{
+	CreateToolGroupManagerFunc = MoveTemp(Func);
+}
+
 
 void UInteractiveToolsContext::SetShutdownInputRouterFunc(TUniqueFunction<void(UInputRouter*)> Func)
 {
@@ -116,6 +134,11 @@ void UInteractiveToolsContext::SetShutdownContextStoreFunc(TUniqueFunction<void(
 	ShutdownContextStoreFunc = MoveTemp(Func);
 }
 
+void UInteractiveToolsContext::SetShutdownToolGroupManagerFunc(TUniqueFunction<void(UToolGroupManager*)> Func)
+{
+	ShutdownToolGroupManagerFunc = MoveTemp(Func);
+}
+
 
 void UInteractiveToolsContext::Initialize(IToolsContextQueriesAPI* QueriesAPI, IToolsContextTransactionsAPI* TransactionsAPI)
 {
@@ -132,11 +155,15 @@ void UInteractiveToolsContext::Initialize(IToolsContextQueriesAPI* QueriesAPI, I
 	ToolManager = CreateToolManagerFunc(InitInfo);
 	TargetManager = CreateToolTargetManagerFunc(InitInfo);
 	GizmoManager = CreateGizmoManagerFunc(InitInfo);
+	ToolGroupManager = CreateToolGroupManagerFunc(InitInfo);
 }
 
 
 void UInteractiveToolsContext::Shutdown()
 {
+	ShutdownToolGroupManagerFunc(ToolGroupManager);
+	ToolGroupManager = nullptr;
+	
 	ShutdownInputRouterFunc(InputRouter);
 	InputRouter = nullptr;
 
diff --git a/Engine/Source/Runtime/InteractiveToolsFramework/Private/ToolGroupManager.cpp b/Engine/Source/Runtime/InteractiveToolsFramework/Private/ToolGroupManager.cpp
new file mode 100644
--- /dev/null
+++ b/Engine/Source/Runtime/InteractiveToolsFramework/Private/ToolGroupManager.cpp
@@ -0,0 +1,332 @@
+﻿// Copyright Epic Games, Inc. All Rights Reserved.
+
+#include "ToolGroupManager.h"
+
+#include "ContextObjectStore.h"
+#include "InteractiveToolsContext.h"
+#include "InputBehaviorSet.h"
+#include "InputRouter.h"
+
+namespace
+{
+	/** Simple adapter that exposes a pre-built behavior set through the IInputBehaviorSource interface. */
+	class FBehaviorSetSource : public IInputBehaviorSource
+	{
+	public:
+		explicit FBehaviorSetSource(const UInputBehaviorSet* InSet)
+			: Set(InSet)
+		{}
+
+		virtual const UInputBehaviorSet* GetInputBehaviors() const override
+		{
+			return Set;
+		}
+
+	private:
+		const UInputBehaviorSet* Set = nullptr;
+	};
+}
+
+void UToolGroupManager::Initialize()
+{
+}
+
+void UToolGroupManager::Shutdown()
+{
+	DisableAllToolGroups();
+	RegisteredGroups.Empty();
+	EnabledGroups.Empty();
+}
+
+bool UToolGroupManager::RegisterToolGroup(FName GroupName, TSharedPtr<const FToolGroup> ToolGroup)
+{
+	if (GroupName.IsNone() || !ToolGroup.IsValid())
+	{
+		return false;
+	}
+
+	if (RegisteredGroups.Contains(GroupName))
+	{
+		return false;
+	}
+
+	RegisteredGroups.Add(GroupName, ToolGroup);
+	return true;
+}
+
+bool UToolGroupManager::HasToolGroup(FName GroupName) const
+{
+	return RegisteredGroups.Contains(GroupName);
+}
+
+bool UToolGroupManager::IsToolGroupEnabled(FName GroupName) const
+{
+	return EnabledGroups.Contains(GroupName);
+}
+
+TArray<FName> UToolGroupManager::GetEnabledToolGroupNames() const
+{
+	TArray<FName> Names;
+	EnabledGroups.GetKeys(Names);
+	return Names;
+}
+
+TOptional<FName> UToolGroupManager::FindEnabledGroupOwningActiveTool(EToolSide Side) const
+{
+	UInteractiveToolsContext* Context = GetOwningContext();
+	if (!Context || !Context->HasActiveTool(Side))
+	{
+		return {};
+	}
+
+	const FString ActiveTool = Context->GetActiveToolName(Side);
+
+	for (const TPair<FName, FEnabledToolGroupState>& Pair : EnabledGroups)
+	{
+		if (Pair.Value.RegisteredToolTypeIds.Contains(ActiveTool))
+		{
+			return Pair.Key;
+		}
+	}
+	return {};
+}
+
+bool UToolGroupManager::HasAnyEnabledToolGroups()
+{
+	return EnabledGroups.Num() > 0;
+}
+
+bool UToolGroupManager::EnableToolGroup(FName GroupName)
+{
+	UInteractiveToolsContext* Context = GetOwningContext();
+	if (!Context || !Context->ToolManager)
+	{
+		return false;
+	}
+
+	if (EnabledGroups.Contains(GroupName))
+	{
+		return true;
+	}
+
+	TSharedPtr<const FToolGroup>* Found = RegisteredGroups.Find(GroupName);
+	if (!Found || !Found->IsValid())
+	{
+		return false;
+	}
+
+	const FToolGroup& ToolGroup = *Found->Get();
+	UInteractiveToolManager* ToolManager = Context->ToolManager;
+
+	FEnabledToolGroupState State;
+
+	// Register tools for this group (record exact tool type ids for ownership + cleanup)
+	RegisterGroupTools(ToolGroup, State, *ToolManager);
+
+	// Acquire shared features into this context (ref-counted)
+	AcquireGroupFeatures(ToolGroup, State, *Context);
+
+	// Build & register per-context behaviors
+	RegisterGroupBehaviors(ToolGroup, State, *Context);
+
+	EnabledGroups.Add(GroupName, MoveTemp(State));
+	return true;
+}
+
+bool UToolGroupManager::DisableToolGroup(FName GroupName)
+{
+	UInteractiveToolsContext* Context = GetOwningContext();
+	if (!Context || !Context->ToolManager)
+	{
+		return false;
+	}
+
+	FEnabledToolGroupState* State = EnabledGroups.Find(GroupName);
+	if (!State)
+	{
+		return false;
+	}
+
+	TSharedPtr<const FToolGroup>* Found = RegisteredGroups.Find(GroupName);
+	if (!Found || !Found->IsValid())
+	{
+		// Group was enabled, but descriptor disappeared (shouldn't happen). Best-effort cleanup.
+		UnregisterGroupBehaviors(*State, *Context);
+		UnregisterGroupTools(*State, *Context->ToolManager);
+		ReleaseAllFeaturesForGroup(*State, *Context);
+		EnabledGroups.Remove(GroupName);
+		return true;
+	}
+
+	const FToolGroup& Descriptor = *Found->Get();
+
+	// If an active tool belongs to this group, cancel it
+	if (Context->HasActiveTool(EToolSide::Mouse))
+	{
+		const FString ActiveTool = Context->GetActiveToolName(EToolSide::Mouse);
+		if (State->RegisteredToolTypeIds.Contains(ActiveTool))
+		{
+			Context->EndTool(EToolSide::Mouse, EToolShutdownType::Cancel);
+		}
+	}
+	
+	// Unregister behaviors
+	UnregisterGroupBehaviors(*State, *Context);
+
+	// Unregister tools
+	UnregisterGroupTools(*State, *Context->ToolManager);
+
+	// Release shared features acquired for this group
+	ReleaseAllFeaturesForGroup(*State, *Context);
+
+	EnabledGroups.Remove(GroupName);
+	return true;
+}
+
+void UToolGroupManager::DisableAllToolGroups()
+{
+	TArray<FName> Names;
+	EnabledGroups.GetKeys(Names);
+	for (FName Name : Names)
+	{
+		DisableToolGroup(Name);
+	}
+}
+void UToolGroupManager::AcquireGroupFeatures(
+	const FToolGroup& ToolGroup,
+	FEnabledToolGroupState& State,
+	UInteractiveToolsContext& Context)
+{
+	ToolGroup.EnumerateFeatures([&](const FToolGroup::FFeature& Feature)
+	{
+		FeatureRegistry.Acquire(Feature.Key, Context, Feature.Install, Feature.Uninstall);
+		State.AcquiredFeatures.Add(Feature.Key);
+		return true;
+	});
+}
+
+void UToolGroupManager::ReleaseAllFeaturesForGroup(FEnabledToolGroupState& State, UInteractiveToolsContext& Context)
+{
+	for (FToolGroupFeatureKey Key : State.AcquiredFeatures)
+	{
+		FeatureRegistry.Release(Key, Context);
+	}
+	State.AcquiredFeatures.Reset();
+}
+
+void UToolGroupManager::RegisterGroupTools(const FToolGroup& ToolGroup, FEnabledToolGroupState& State, UInteractiveToolManager& ToolManager)
+{
+	ToolGroup.EnumerateTools([&](const FTool& Tool)
+	{
+		if (!Tool.BuilderClass)
+		{
+			return true;
+		}
+
+		// If/when namespacing is added later, compute ToolTypeId here and store that.
+		const FString ToolTypeId = Tool.Name;
+
+		UInteractiveToolBuilder* Builder = NewObject<UInteractiveToolBuilder>(&ToolManager, Tool.BuilderClass);
+		if (!ensureMsgf(Builder, TEXT("Failed to allocate tool builder for ToolTypeId '%s' (%s)"), *ToolTypeId, *GetNameSafe(Tool.BuilderClass)))
+		{
+			return true; // continue enumeration but skip this tool
+		}
+		ToolManager.RegisterToolType(ToolTypeId, Builder);
+		State.RegisteredToolTypeIds.Add(ToolTypeId);
+		return true;
+	});
+}
+
+void UToolGroupManager::UnregisterGroupTools(FEnabledToolGroupState& State, UInteractiveToolManager& ToolManager)
+{
+	for (const FString& ToolTypeId : State.RegisteredToolTypeIds)
+	{
+		ToolManager.UnregisterToolType(ToolTypeId);
+	}
+	State.RegisteredToolTypeIds.Reset();
+}
+
+void UToolGroupManager::RegisterGroupBehaviors(const FToolGroup& ToolGroup, FEnabledToolGroupState& State, UInteractiveToolsContext& Context)
+{
+	if (!Context.InputRouter)
+	{
+		return;
+	}
+
+	if (!IsValid(State.BehaviorSet))
+	{
+		// Outer on the manager so the behavior set's lifetime is tied to this context's manager instance.
+		State.BehaviorSet = NewObject<UInputBehaviorSet>(this);
+		ToolGroup.BuildInputBehaviors(Context, *State.BehaviorSet.Get());
+	}
+
+	if (!State.bRegisteredAsBehaviorSource)
+	{
+		State.BehaviorSource = MakeShared<FBehaviorSetSource>(State.BehaviorSet.Get());
+		Context.InputRouter->RegisterSource(State.BehaviorSource.Get());
+		State.bRegisteredAsBehaviorSource = true;
+	}
+}
+
+void UToolGroupManager::UnregisterGroupBehaviors(FEnabledToolGroupState& State, UInteractiveToolsContext& Context)
+{
+	if (Context.InputRouter && State.bRegisteredAsBehaviorSource && State.BehaviorSource.IsValid())
+	{
+		Context.InputRouter->DeregisterSource(State.BehaviorSource.Get());
+	}
+
+	State.bRegisteredAsBehaviorSource = false;
+	State.BehaviorSource.Reset();
+	State.BehaviorSet = nullptr;
+}
+UInteractiveToolsContext* UToolGroupManager::GetOwningContext() const
+{
+	return Cast<UInteractiveToolsContext>(GetOuter());
+}
+
+void FToolGroupFeatureRegistry::Acquire(
+	FToolGroupFeatureKey Key,
+	UInteractiveToolsContext& Context,
+	TFunction<void(UInteractiveToolsContext&)> Install,
+	TFunction<void(UInteractiveToolsContext&)> Uninstall)
+{
+	FToolGroupFeatureEntry& Entry = Features.FindOrAdd(Key);
+
+	// First acquisition installs and defines callbacks for this context.
+	if (Entry.RefCount == 0)
+	{
+		Entry.Install = MoveTemp(Install);
+		Entry.Uninstall = MoveTemp(Uninstall);
+
+		if (Entry.Install)
+		{
+			Entry.Install(Context);
+		}
+	}
+	else
+	{
+		// Re-acquire: feature already defined for this context; ignore callbacks and just ref-count.
+		// (Immutable ToolGroup descriptors will naturally re-provide callbacks here.)
+	}
+
+	++Entry.RefCount;
+}
+
+void FToolGroupFeatureRegistry::Release(FToolGroupFeatureKey Key, UInteractiveToolsContext& Context)
+{
+	FToolGroupFeatureEntry* Entry = Features.Find(Key);
+	if (!Entry)
+	{
+		return;
+	}
+
+	Entry->RefCount--;
+	if (Entry->RefCount <= 0)
+	{
+		if (Entry->Uninstall)
+		{
+			Entry->Uninstall(Context);
+		}
+		Features.Remove(Key);
+	}
+}
diff --git a/Engine/Source/Runtime/InteractiveToolsFramework/Public/InteractiveToolsContext.h b/Engine/Source/Runtime/InteractiveToolsFramework/Public/InteractiveToolsContext.h
--- a/Engine/Source/Runtime/InteractiveToolsFramework/Public/InteractiveToolsContext.h
+++ b/Engine/Source/Runtime/InteractiveToolsFramework/Public/InteractiveToolsContext.h
@@ -17,6 +17,7 @@
 
 #include "InteractiveToolsContext.generated.h"
 
+class UToolGroupManager;
 class FText;
 class IToolsContextQueriesAPI;
 class IToolsContextTransactionsAPI;
@@ -28,7 +29,7 @@ class UToolTargetManager;
 
 /**
  * InteractiveToolsContext owns the core parts of an Interactive Tools Framework implementation - the
- * InputRouter, ToolManager, GizmoManager, TargetManager, and ContextStore. In the simplest
+ * InputRouter, ToolManager, GizmoManager, TargetManager, ToolGroupManager, and ContextStore. In the simplest
  * use case, UInteractiveToolsContext is just a top-level container that will keep the various UObjects
  * alive, and provide an easy way to access them. However in a more complex situation it may be
  * desirable to subclass and extend the ToolsContext. For example, UEdModeInteractiveToolsContext allows
@@ -63,12 +64,14 @@ public:
 	INTERACTIVETOOLSFRAMEWORK_API virtual void SetCreateToolTargetManagerFunc(TUniqueFunction<UToolTargetManager* (const FContextInitInfo&)> Func);
 	INTERACTIVETOOLSFRAMEWORK_API virtual void SetCreateGizmoManagerFunc(TUniqueFunction<UInteractiveGizmoManager* (const FContextInitInfo&)> Func);
 	INTERACTIVETOOLSFRAMEWORK_API virtual void SetCreateContextStoreFunc(TUniqueFunction<UContextObjectStore* (const FContextInitInfo&)> Func);
+	INTERACTIVETOOLSFRAMEWORK_API virtual void SetCreateToolGroupManagerFunc(TUniqueFunction<UToolGroupManager* (const FContextInitInfo&)> Func);
 
 	INTERACTIVETOOLSFRAMEWORK_API virtual void SetShutdownInputRouterFunc(TUniqueFunction<void(UInputRouter*)> Func);
 	INTERACTIVETOOLSFRAMEWORK_API virtual void SetShutdownToolManagerFunc(TUniqueFunction<void(UInteractiveToolManager*)> Func);
 	INTERACTIVETOOLSFRAMEWORK_API virtual void SetShutdownToolTargetManagerFunc(TUniqueFunction<void(UToolTargetManager*)> Func);
 	INTERACTIVETOOLSFRAMEWORK_API virtual void SetShutdownGizmoManagerFunc(TUniqueFunction<void(UInteractiveGizmoManager*)> Func);
 	INTERACTIVETOOLSFRAMEWORK_API virtual void SetShutdownContextStoreFunc(TUniqueFunction<void(UContextObjectStore*)> Func);
+	INTERACTIVETOOLSFRAMEWORK_API virtual void SetShutdownToolGroupManagerFunc(TUniqueFunction<void(UToolGroupManager*)> Func);
 
 	/** 
 	 * Initialize the Context. This creates the InputRouter, ToolManager, GizmoManager, TargetManager, and ContextStore
@@ -120,6 +123,10 @@ public:
 	UPROPERTY()
 	TObjectPtr<UInteractiveGizmoManager> GizmoManager;
 
+	/** current UToolGroupManager for this Context */
+	UPROPERTY()
+	TObjectPtr<UToolGroupManager> ToolGroupManager;
+	
 	/** 
 	 * Current Context Object Store for this Context.
 	 * Stores arbitrary objects which share data or expose APIs across interactive tools and managers belonging to this context.
@@ -141,6 +148,7 @@ protected:
 	TUniqueFunction<UToolTargetManager* (const FContextInitInfo&)> CreateToolTargetManagerFunc;
 	TUniqueFunction<UInteractiveGizmoManager* (const FContextInitInfo&)> CreateGizmoManagerFunc;
 	TUniqueFunction<UContextObjectStore* (const FContextInitInfo&)> CreateContextStoreFunc;
+	TUniqueFunction<UToolGroupManager* (const FContextInitInfo&)> CreateToolGroupManagerFunc;
 
 	// Analogous to the CreateX() functions above, these function are called by Shutdown()
 	// to terminate and clean up after the various elements. 
@@ -149,6 +157,7 @@ protected:
 	TUniqueFunction<void(UToolTargetManager*)> ShutdownToolTargetManagerFunc;
 	TUniqueFunction<void(UInteractiveGizmoManager*)> ShutdownGizmoManagerFunc;
 	TUniqueFunction<void(UContextObjectStore*)> ShutdownContextStoreFunc;
+	TUniqueFunction<void(UToolGroupManager*)> ShutdownToolGroupManagerFunc;
 
 
 	// todo: deprecate and remove this, can now be accomplished via CreateToolManagerFunc()
diff --git a/Engine/Source/Runtime/InteractiveToolsFramework/Public/ToolGroup.h b/Engine/Source/Runtime/InteractiveToolsFramework/Public/ToolGroup.h
new file mode 100644
--- /dev/null
+++ b/Engine/Source/Runtime/InteractiveToolsFramework/Public/ToolGroup.h
@@ -0,0 +1,114 @@
+﻿// Copyright Epic Games, Inc. All Rights Reserved.
+
+#pragma once
+
+#include "CoreMinimal.h"
+#include "InputBehaviorSet.h"
+#include "InteractiveToolBuilder.h"
+#include "HAL/Platform.h"
+
+class UInteractiveToolsContext;
+
+/** Describes a tool type to be registered into a UInteractiveToolManager. */
+UE_EXPERIMENTAL(5.8, "Experimental, API may change in the future")
+struct FTool
+{
+	/** ITF Tool Identifier (ToolManager tool type id) */
+	FString Name;
+
+	/** Builder class registered with the ToolManager. */
+	TSubclassOf<UInteractiveToolBuilder> BuilderClass;
+};
+/**
+ * Declarative description of a ToolGroup.
+ *
+ * A ToolGroup describes:
+ *  - which tool types should be registered (FTool entries),
+ *  - which shared, context-wide features should be installed (Install/Uninstall callbacks),
+ *  - which input behaviors should be active (built into a per-context UInputBehaviorSet),
+ *  - optional context policy hooks (eg editor-only widget visibility).
+ *
+ * ToolGroups are intentionally *immutable descriptors* and must not own per-context runtime state.
+ * All per-context state (registered tool ids, acquired features, behavior sets, etc.) is owned by
+ * UToolGroupManager in FEnabledToolGroupState.
+ */
+UE_EXPERIMENTAL(5.8, "Experimental, API may change in the future")
+class FToolGroup
+{
+public:
+	/** A shared, reference-counted feature installed into a UInteractiveToolsContext. */
+	struct FFeature
+	{
+		/** Stable key identifying the feature within a tools context. */
+		FName Key;
+
+		/** Called on 0->1 ref-count transition. */
+		TFunction<void(UInteractiveToolsContext&)> Install;
+
+		/** Called on 1->0 ref-count transition. */
+		TFunction<void(UInteractiveToolsContext&)> Uninstall;
+	};
+
+	/** Hook used to build the per-context input behaviors for this group. */
+	using FBuildBehaviorsHook = TFunction<void(UInteractiveToolsContext& ToolsContext, UInputBehaviorSet& BehaviorSet)>;
+
+	virtual ~FToolGroup() = default;
+
+	/** Enumerate tools in this group. Enumeration continues while ToolFunc returns true. */
+	virtual void EnumerateTools(TFunction<bool(const FTool&)> ToolFunc) const
+	{
+		if (!ToolFunc)
+		{
+			return;
+		}
+		for (const FTool& Tool : Tools)
+		{
+			if (!ToolFunc(Tool))
+			{
+				break;
+			}
+		}
+	}
+
+	/** Enumerate shared feature requirements for this group. */
+	virtual void EnumerateFeatures(TFunction<bool(const FFeature&)> FeatureFunc) const
+	{
+		if (!FeatureFunc)
+		{
+			return;
+		}
+		for (const FFeature& Feature : Features)
+		{
+			if (!FeatureFunc(Feature))
+			{
+				break;
+			}
+		}
+	}
+
+	/**
+	 * Populate per-context behaviors into the provided BehaviorSet.
+	 *
+	 * Behavior instances are generally not safe to share across tools contexts. The ToolGroupManager
+	 * will allocate a new BehaviorSet per enabled group per context, call this hook, and then register
+	 * a behavior source with the context's InputRouter.
+	 */
+	virtual void BuildInputBehaviors(UInteractiveToolsContext& ToolsContext, UInputBehaviorSet& BehaviorSet) const
+	{
+		if (BuildBehaviorsFunc)
+		{
+			BuildBehaviorsFunc(ToolsContext, BehaviorSet);
+		}
+	}
+
+
+public:
+	/** Declarative tool list. */
+	TArray<FTool> Tools;
+
+	/** Declarative feature requirements. */
+	TArray<FFeature> Features;
+
+	/** Optional behavior builder (called per enabled group per context). */
+	FBuildBehaviorsHook BuildBehaviorsFunc;
+};
\ No newline at end of file
diff --git a/Engine/Source/Runtime/InteractiveToolsFramework/Public/ToolGroupManager.h b/Engine/Source/Runtime/InteractiveToolsFramework/Public/ToolGroupManager.h
new file mode 100644
--- /dev/null
+++ b/Engine/Source/Runtime/InteractiveToolsFramework/Public/ToolGroupManager.h
@@ -0,0 +1,194 @@
+﻿// Copyright Epic Games, Inc. All Rights Reserved.
+
+#pragma once
+
+#include "CoreMinimal.h"
+#include "InteractiveToolManager.h"
+#include "ToolGroup.h"
+#include "UObject/Object.h"
+#include "ToolGroupManager.generated.h"
+
+class UInteractiveToolsContext;
+class IInputBehaviorSource;
+
+#define UE_API INTERACTIVETOOLSFRAMEWORK_API
+
+/** Stable identifier for a shared feature installed into a UInteractiveToolsContext. */
+using FToolGroupFeatureKey = FName;
+
+/**
+ * Stored state for a single acquired feature key.
+ *
+ * Features are context-wide singleton-style services (target factories, snapping managers, etc.)
+ * installed into a UInteractiveToolsContext and shared across multiple enabled tool groups.
+ */
+UE_EXPERIMENTAL(5.8, "Experimental, API may change in the future")
+struct FToolGroupFeatureEntry
+{
+	/**
+	 * Current number of enabled ToolGroups that depend on this feature.
+	 *
+	 * RefCount is managed by FToolGroupFeatureRegistry and should not be manipulated directly.
+	 */
+	int32 RefCount = 0;
+
+	/** Called only when RefCount transitions 0 -> 1. */
+	TFunction<void(UInteractiveToolsContext&)> Install;
+
+	/** Called only when RefCount transitions 1 -> 0. */
+	TFunction<void(UInteractiveToolsContext&)> Uninstall;
+};
+
+/**
+ * Reference-counted registry for context-wide features.
+ *
+ * This lives per UToolGroupManager (ie per UInteractiveToolsContext). Multiple groups can request
+ * the same feature key; the first request installs it into the context and the final release
+ * uninstalls it.
+ */
+UE_EXPERIMENTAL(5.8, "Experimental, API may change in the future")
+class FToolGroupFeatureRegistry
+{
+public:
+	/**
+	 * Acquire the feature. If this is the first acquire for Key (0 -> 1), Install will be executed and the
+	 * provided Install/Uninstall callbacks will be stored as the feature definition for this context.
+	 *
+	 * If the feature is already acquired (RefCount > 0), the stored Install/Uninstall callbacks are retained.
+	 * Callers may pass callbacks again (eg from immutable ToolGroup descriptors); these callbacks are ignored
+	 * and only the ref-count is incremented.
+	 *
+	 * Note: Feature keys are expected to uniquely identify feature behavior within a context. If multiple
+	 * callers provide different implementations for the same key, behavior is defined by the first acquire.
+	 */
+	void Acquire(
+		FToolGroupFeatureKey Key,
+		UInteractiveToolsContext& Context,
+		TFunction<void(UInteractiveToolsContext&)> Install,
+		TFunction<void(UInteractiveToolsContext&)> Uninstall);
+
+	/** Release a previously acquired feature. Final release triggers Uninstall. */
+	void Release(FToolGroupFeatureKey Key, UInteractiveToolsContext& Context);
+
+	/** True if the feature key is currently acquired by at least one group. */
+	bool IsAcquired(FToolGroupFeatureKey Key) const { return Features.Contains(Key); }
+
+private:
+	TMap<FToolGroupFeatureKey, FToolGroupFeatureEntry> Features;
+};
+
+/**
+ * Per-context runtime state for an enabled ToolGroup.
+ *
+ * ToolGroups themselves are immutable descriptors. All per-context state is owned here by the
+ * per-context UToolGroupManager, including:
+ *  - which tool type ids were registered for this group,
+ *  - which shared features were acquired for this group,
+ *  - the per-context UInputBehaviorSet constructed for this group,
+ *  - the behavior source adapter registered with the context's InputRouter.
+ */
+UE_EXPERIMENTAL(5.8, "Experimental, API may change in the future")
+USTRUCT()
+struct FEnabledToolGroupState
+{
+	GENERATED_BODY()
+	
+	/** Per-context behavior set for this enabled group. */
+	UPROPERTY(transient)
+	TObjectPtr<UInputBehaviorSet> BehaviorSet = nullptr;
+	
+	/** Tool type identifiers registered with the ToolManager on behalf of this group. */
+	TSet<FString> RegisteredToolTypeIds;
+
+	/** Features acquired by this enabled group (explicit tracking for cleanup). */
+	TSet<FToolGroupFeatureKey> AcquiredFeatures;
+
+	/** Adapter registered with the InputRouter for this group. */
+	TSharedPtr<IInputBehaviorSource> BehaviorSource;
+
+	/** True if BehaviorSource has been registered with the owning context's InputRouter. */
+	bool bRegisteredAsBehaviorSource = false;
+};
+
+/**
+ * Manages ToolGroups for a single UInteractiveToolsContext.
+ *
+ * UToolGroupManager is stored per-tools-context and is responsible for 
+ * materializing per-context runtime state from immutable tool group descriptors.
+ *
+ * Responsibilities:
+ *  - register immutable ToolGroup descriptors by name,
+ *  - enable/disable groups within this tools context,
+ *  - register/unregister tool types with UInteractiveToolManager,
+ *  - acquire/release shared features via FToolGroupFeatureRegistry,
+ *  - build/register per-context input behaviors with the context's InputRouter,
+ *  - answer queries like "which enabled group owns the currently active tool".
+ */
+UE_EXPERIMENTAL(5.8, "Experimental, API may change in the future")
+UCLASS(Transient)
+class UToolGroupManager : public UObject
+{
+	GENERATED_BODY()
+
+public:
+	/** Initialize internal state. Must be called before registering/enabling groups. */
+	UE_API void Initialize();
+	
+	/** Shutdown and disable all enabled tool groups, unregistering their tool types. */
+	UE_API void Shutdown();
+	
+	/** Register an immutable ToolGroup descriptor that can later be enabled in this context. */
+	UE_API bool RegisterToolGroup(FName GroupName, TSharedPtr<const FToolGroup> ToolGroup);
+
+	/** Returns true if a tool group descriptor with this name is registered. */
+	UE_API bool HasToolGroup(FName GroupName) const;
+
+	/** Enable (materialize + register) the named tool group in this context. */
+	UE_API bool EnableToolGroup(FName GroupName);
+
+	/** Disable (unregister + cleanup) the named tool group in this context. */
+	UE_API bool DisableToolGroup(FName GroupName);
+
+	/** Disable all enabled tool groups in this context. */
+	UE_API void DisableAllToolGroups();
+
+	// Lifecycle
+	/** Returns true if the named group is currently enabled in this context. */
+	UE_API bool IsToolGroupEnabled(FName GroupName) const;
+
+	/** Returns the names of currently enabled groups. */
+	UE_API TArray<FName> GetEnabledToolGroupNames() const;
+
+
+	// Tool Ownership
+	/** If a tool is currently active, returns the name of the enabled group that registered it. */
+	UE_API TOptional<FName> FindEnabledGroupOwningActiveTool(EToolSide Side) const;
+
+	/** Returns true if any tool groups are enabled in this context. */
+	UE_API bool HasAnyEnabledToolGroups();
+
+private:
+	void AcquireGroupFeatures(const FToolGroup& ToolGroup, FEnabledToolGroupState& State, UInteractiveToolsContext& Context);
+	void ReleaseAllFeaturesForGroup(FEnabledToolGroupState& State, UInteractiveToolsContext& Context);
+
+	void RegisterGroupTools(const FToolGroup& ToolGroup, FEnabledToolGroupState& State, UInteractiveToolManager& ToolManager);
+	void UnregisterGroupTools(FEnabledToolGroupState& State, UInteractiveToolManager& ToolManager);
+
+	void RegisterGroupBehaviors(const FToolGroup& ToolGroup, FEnabledToolGroupState& State, UInteractiveToolsContext& Context);
+	void UnregisterGroupBehaviors(FEnabledToolGroupState& State, UInteractiveToolsContext& Context);
+
+	UInteractiveToolsContext* GetOwningContext() const;
+
+private:
+	/** Registered immutable tool group descriptors by name. */
+	TMap<FName, TSharedPtr<const FToolGroup>> RegisteredGroups;
+
+	/** Per-context runtime state for enabled groups. */
+	UPROPERTY(Transient)
+	TMap<FName, FEnabledToolGroupState> EnabledGroups;
+
+	/** Ref-counted context-wide feature registry. */
+	FToolGroupFeatureRegistry FeatureRegistry;
+};
+
+#undef UE_API
\ No newline at end of file
