Arc Ball - Drawing order shows clipping

#Synced-CL 51038616

diff --git a/Engine/Source/Runtime/InteractiveToolsFramework/Private/BaseGizmos/GizmoElementCircle.cpp b/Engine/Source/Runtime/InteractiveToolsFramework/Private/BaseGizmos/GizmoElementCircle.cpp
--- a/Engine/Source/Runtime/InteractiveToolsFramework/Private/BaseGizmos/GizmoElementCircle.cpp
+++ b/Engine/Source/Runtime/InteractiveToolsFramework/Private/BaseGizmos/GizmoElementCircle.cpp
@@ -116,10 +116,12 @@ void UGizmoElementCircle::Render(IToolsContextRenderAPI* RenderAPI, const FRende
 			{
 				using namespace UE::InteractiveToolsFramework::Private;
 
-				// Offset in view-space to avoid z-fighting with the line
+				// Offset in view-space to avoid z-fighting with the line.
+				// Note: UpdateRenderState has already applied ViewDepthOffset to CurrentRenderState.LocalToWorldTransform,
+				// so we only need to apply the additional offset to push the mesh behind the line.
 				if (bDrawLine)
 				{
-					const double MeshViewDepthOffset = GetViewDepthOffset() - CurrentLineThickness;
+					const double MeshViewDepthOffset = -CurrentLineThickness;
 					ApplyViewDepthOffset(View, MeshViewDepthOffset, CurrentRenderState.LocalToWorldTransform);
 				}
 
