Modeling Tools: enable 'specify radius' option in Paint Maps tool by default

#Synced-CL 50957382

diff --git a/Engine/Plugins/Runtime/MeshModelingToolset/Source/MeshModelingTools/Private/MeshAttributePaintTool.cpp b/Engine/Plugins/Runtime/MeshModelingToolset/Source/MeshModelingTools/Private/MeshAttributePaintTool.cpp
--- a/Engine/Plugins/Runtime/MeshModelingToolset/Source/MeshModelingTools/Private/MeshAttributePaintTool.cpp
+++ b/Engine/Plugins/Runtime/MeshModelingToolset/Source/MeshModelingTools/Private/MeshAttributePaintTool.cpp
@@ -323,6 +323,7 @@ void UMeshAttributePaintTool::Setup()
 	{
 		BrushProperties->bToolSupportsPressureSensitivity = true;
 		BrushProperties->bShowHitBackFaces = true;
+		BrushProperties->bSpecifyRadius = true;
 	}
 
 	// hide strength and falloff
diff --git a/Engine/Source/Runtime/InteractiveToolsFramework/Public/BaseTools/BaseBrushTool.h b/Engine/Source/Runtime/InteractiveToolsFramework/Public/BaseTools/BaseBrushTool.h
--- a/Engine/Source/Runtime/InteractiveToolsFramework/Public/BaseTools/BaseBrushTool.h
+++ b/Engine/Source/Runtime/InteractiveToolsFramework/Public/BaseTools/BaseBrushTool.h
@@ -32,7 +32,7 @@ public:
 
 	/** Radius of brush */
 	UPROPERTY(EditAnywhere, Category = Brush, AdvancedDisplay, meta = (EditCondition = "bSpecifyRadius",
-		DisplayName = "Radius", UIMin = "1.0", UIMax = "1000.0", ClampMin = "0.1", ClampMax = "50000.0"))
+		DisplayName = "Radius", UIMin = "1.0", UIMax = "1000.0", ClampMin = "0.1", ClampMax = "50000.0", ModelingQuickSettings = 100))
 	float BrushRadius;
 
 	//~Set by tool and determines whether the bEnablePressureSensitivity flag below is shown and used.
