[ITF] Drop any active mouse hover capture prior to calling `BeginCapture()`. Fix issue where terminating hover cleared the wrong owner field.

[RN] Minor | Editor

#Synced-CL 51211095

diff --git a/Engine/Source/Runtime/InteractiveToolsFramework/Private/InputRouter.cpp b/Engine/Source/Runtime/InteractiveToolsFramework/Private/InputRouter.cpp
--- a/Engine/Source/Runtime/InteractiveToolsFramework/Private/InputRouter.cpp
+++ b/Engine/Source/Runtime/InteractiveToolsFramework/Private/InputRouter.cpp
@@ -292,6 +292,9 @@ void UInputRouter::CheckForMouseCaptures(const FInputDeviceState& Input)
 
 	CaptureRequests.StableSort();
 
+	// end outstanding hover
+	TerminateHover(EInputCaptureSide::Left);
+
 	bool bAccepted = false;
 	for (int i = 0; i < CaptureRequests.Num() && bAccepted == false; ++i)
 	{
@@ -299,9 +302,6 @@ void UInputRouter::CheckForMouseCaptures(const FInputDeviceState& Input)
 			CaptureRequests[i].Source->BeginCapture(Input, EInputCaptureSide::Left);
 		if (Result.State == EInputCaptureState::Begin)
 		{
-			// end outstanding hover
-			TerminateHover(EInputCaptureSide::Left);
-
 			ActiveLeftCapture = Result.Source;
 			ActiveLeftCaptureOwner = CaptureRequests[i].Owner;
 			ActiveLeftCaptureData = Result.Data;
@@ -399,7 +399,7 @@ void UInputRouter::TerminateHover(EInputCaptureSide Side)
 	{
 		ActiveLeftHoverCapture->EndHoverCapture();
 		ActiveLeftHoverCapture = nullptr;
-		ActiveLeftCaptureOwner = nullptr;
+		ActiveLeftHoverCaptureOwner = nullptr;
 	}
 }
 
