[ITF] UKeyInputBehaviors that always update modifiers will also return true from WantsForceEndCapture. This ensures that any state derived from non-captured activity is cleaned out at the appropriate times.

#rnx
#Synced-CL 51461249

diff --git a/Engine/Source/Runtime/InteractiveToolsFramework/Private/BaseBehaviors/KeyInputBehavior.cpp b/Engine/Source/Runtime/InteractiveToolsFramework/Private/BaseBehaviors/KeyInputBehavior.cpp
--- a/Engine/Source/Runtime/InteractiveToolsFramework/Private/BaseBehaviors/KeyInputBehavior.cpp
+++ b/Engine/Source/Runtime/InteractiveToolsFramework/Private/BaseBehaviors/KeyInputBehavior.cpp
@@ -100,6 +100,11 @@ void UKeyInputBehavior::ForceEndCapture(const FInputCaptureData& Data)
 	}
 }
 
+bool UKeyInputBehavior::WantsForceEndCapture()
+{
+	return bAlwaysUpdateModifiers;
+}
+
 bool UKeyInputBehavior::IsKeyPressed(FKey InKey) const
 {
 	const int32 KeyIndex = TargetKeys.Find(InKey);
diff --git a/Engine/Source/Runtime/InteractiveToolsFramework/Public/BaseBehaviors/KeyInputBehavior.h b/Engine/Source/Runtime/InteractiveToolsFramework/Public/BaseBehaviors/KeyInputBehavior.h
--- a/Engine/Source/Runtime/InteractiveToolsFramework/Public/BaseBehaviors/KeyInputBehavior.h
+++ b/Engine/Source/Runtime/InteractiveToolsFramework/Public/BaseBehaviors/KeyInputBehavior.h
@@ -78,6 +78,7 @@ public:
 	INTERACTIVETOOLSFRAMEWORK_API virtual FInputCaptureUpdate BeginCapture(const FInputDeviceState& Input, EInputCaptureSide Side) override;
 	INTERACTIVETOOLSFRAMEWORK_API virtual FInputCaptureUpdate UpdateCapture(const FInputDeviceState& Input, const FInputCaptureData& Data) override;
 	INTERACTIVETOOLSFRAMEWORK_API virtual void ForceEndCapture(const FInputCaptureData& Data) override;
+	INTERACTIVETOOLSFRAMEWORK_API virtual bool WantsForceEndCapture() override;
 
 	/**
 	 * @return Whether the given key is detected as pressed.
