SGUI: Initial ITF Framework and interactable PrefabEditor Viewport
- Implements a prototype Selection Tool with a Resize Gizmo
- Ctrl Key Modifier support per-pixel selection

#rnx

#Synced-CL 51980804

diff --git a/Engine/Source/Runtime/InteractiveToolsFramework/Private/InteractiveGizmo.cpp b/Engine/Source/Runtime/InteractiveToolsFramework/Private/InteractiveGizmo.cpp
--- a/Engine/Source/Runtime/InteractiveToolsFramework/Private/InteractiveGizmo.cpp
+++ b/Engine/Source/Runtime/InteractiveToolsFramework/Private/InteractiveGizmo.cpp
@@ -37,6 +37,11 @@ void UInteractiveGizmo::AddInputBehavior(UInputBehavior* Behavior)
 	InputBehaviors->Add(Behavior);
 }
 
+void UInteractiveGizmo::RemoveInputBehavior(UInputBehavior* Behavior)
+{
+	InputBehaviors->Remove(Behavior);
+}
+
 const UInputBehaviorSet* UInteractiveGizmo::GetInputBehaviors() const
 {
 	return InputBehaviors;
diff --git a/Engine/Source/Runtime/InteractiveToolsFramework/Public/InteractiveGizmo.h b/Engine/Source/Runtime/InteractiveToolsFramework/Public/InteractiveGizmo.h
--- a/Engine/Source/Runtime/InteractiveToolsFramework/Public/InteractiveGizmo.h
+++ b/Engine/Source/Runtime/InteractiveToolsFramework/Public/InteractiveGizmo.h
@@ -80,6 +80,12 @@ public:
 	 */
 	INTERACTIVETOOLSFRAMEWORK_API virtual void AddInputBehavior(UInputBehavior* Behavior);
 
+	/**
+	 * Remove an input behavior from this Gizmo
+	 * @param Behavior behavior to add
+	 */
+	INTERACTIVETOOLSFRAMEWORK_API virtual void RemoveInputBehavior(UInputBehavior* Behavior);
+
 	/**
 	 * @return Current input behavior set.
 	 */
