[ITF Gizmos] Inconsistent visuals with Gizmo handles

Handling some issues with gizmo handles visuals.

Split VIewports:
- Wrong elements scale: use `SceneView` instead of `GizmoViewContext`, to make sure to use information from the current viewport.

- Scale handle dashed line scaling issues: `GetScreenSpaceLineStartEnd`, was sampling OutPoints depth at the wrong place (a clamped point not located on the actual dashed line).

- Cut dashed line during Shift + Drag
Added a parameter to `GetScreenSpaceLineStartEnd`, which tracks the length of the line cut portion when its origin is outside the viewport, and uses it to offset the dashed line when it is created.
`ParamToRadiusScaleFunction` for the dash line in `UGizmoElementCylinder::Render` was updated accordingly

- Split view: dragging an object which is outside an ortho view from a perspective view, making it appear in the ortho view, would have a truncated/offset line
Should be fixed by cut dashed line fix.

#Synced-CL 52295285

diff --git a/Engine/Source/Editor/Experimental/EditorInteractiveToolsFramework/Private/EditorGizmos/EditorTRSGizmo.cpp b/Engine/Source/Editor/Experimental/EditorInteractiveToolsFramework/Private/EditorGizmos/EditorTRSGizmo.cpp
--- a/Engine/Source/Editor/Experimental/EditorInteractiveToolsFramework/Private/EditorGizmos/EditorTRSGizmo.cpp
+++ b/Engine/Source/Editor/Experimental/EditorInteractiveToolsFramework/Private/EditorGizmos/EditorTRSGizmo.cpp
@@ -2789,7 +2789,6 @@ UGizmoElementTranslateGroup* UEditorTRSGizmo::MakeTranslateGroup(const EAxisList
 	using namespace UE::Editor::InteractiveToolsFramework::Private;
 
 	UGizmoElementTranslateGroup* Element = NewObject<UGizmoElementTranslateGroup>();
-	Element->SetGizmoViewContext(GizmoViewContext);
 
 	Element->Setup(
 		static_cast<uint32>(ETransformGizmoPartIdentifier::TranslateAll),
@@ -2852,7 +2851,6 @@ UGizmoElementScaleGroup* UEditorTRSGizmo::MakeScaleGroup(const EAxisList::Type I
 	using namespace UE::Editor::InteractiveToolsFramework::Private;
 
 	UGizmoElementScaleGroup* Element = NewObject<UGizmoElementScaleGroup>();
-	Element->SetGizmoViewContext(GizmoViewContext);
 
 	Element->Setup(
 		static_cast<uint32>(ETransformGizmoPartIdentifier::ScaleAll),
@@ -2963,7 +2961,6 @@ UGizmoElementRotateAxisSet* UEditorTRSGizmo::MakeRotateAxisSet(const EAxisList::
 	using namespace UE::Editor::InteractiveToolsFramework::Private;
 
 	UGizmoElementRotateAxisSet* Element = NewObject<UGizmoElementRotateAxisSet>();
-	Element->SetGizmoViewContext(GizmoViewContext);
 
 	Element->Setup(
 		InAxisList,
@@ -3494,21 +3491,6 @@ void UEditorTRSGizmo::UpdateMaterials()
 void UEditorTRSGizmo::SetGizmoViewContext(UGizmoViewContext* InGizmoViewContext)
 {
 	GizmoViewContext = InGizmoViewContext;
-
-	if (TranslateGroupElement)
-	{
-		TranslateGroupElement->SetGizmoViewContext(GizmoViewContext);
-	}
-
-	if (RotateAxisSetElement)
-	{
-		RotateAxisSetElement->SetGizmoViewContext(GizmoViewContext);
-	}
-
-	if (ScaleGroupElement)
-	{
-		ScaleGroupElement->SetGizmoViewContext(GizmoViewContext);
-	}
 }
 
 void UEditorTRSGizmo::SetupOnHoverFunctions()
diff --git a/Engine/Source/Editor/Experimental/EditorInteractiveToolsFramework/Private/EditorGizmos/GizmoElementDashedLine.cpp b/Engine/Source/Editor/Experimental/EditorInteractiveToolsFramework/Private/EditorGizmos/GizmoElementDashedLine.cpp
--- a/Engine/Source/Editor/Experimental/EditorInteractiveToolsFramework/Private/EditorGizmos/GizmoElementDashedLine.cpp
+++ b/Engine/Source/Editor/Experimental/EditorInteractiveToolsFramework/Private/EditorGizmos/GizmoElementDashedLine.cpp
@@ -55,9 +55,10 @@ void UGizmoElementDashedLine::Render(IToolsContextRenderAPI* RenderAPI, const FR
 			// Get screen-space endpoints (used for visibility and for reparameterization, both world and screen-space)
 			FVector ScreenStartPoint, ScreenEndPoint;
 			double LineFractionInView = 1.0;
+			double WorldStartFraction = 0.0;
 
 			GizmoMath::GetScreenSpaceLineStartEnd(
-				RenderAPI->GetSceneView(), WorldOrigin, WorldDirection, WorldLength, ScreenStartPoint, ScreenEndPoint, LineFractionInView
+				RenderAPI->GetSceneView(), WorldOrigin, WorldDirection, WorldLength, ScreenStartPoint, ScreenEndPoint, LineFractionInView, WorldStartFraction
 			);
 
 			GizmoMath::FDashLineSettings DashLineSettings;
@@ -71,6 +72,7 @@ void UGizmoElementDashedLine::Render(IToolsContextRenderAPI* RenderAPI, const FR
 			DashLineSettings.ScreenEndPoint = ScreenEndPoint;
 			DashLineSettings.ScaledLength = ScaledLength;
 			DashLineSettings.LineFractionInView = LineFractionInView;
+			DashLineSettings.WorldStartFraction = WorldStartFraction;
 			DashLineSettings.NumSides = 3;
 
 			if (bScreenSpaceLine)
diff --git a/Engine/Source/Editor/Experimental/EditorInteractiveToolsFramework/Private/EditorGizmos/GizmoElementRotateAxisSet.cpp b/Engine/Source/Editor/Experimental/EditorInteractiveToolsFramework/Private/EditorGizmos/GizmoElementRotateAxisSet.cpp
--- a/Engine/Source/Editor/Experimental/EditorInteractiveToolsFramework/Private/EditorGizmos/GizmoElementRotateAxisSet.cpp
+++ b/Engine/Source/Editor/Experimental/EditorInteractiveToolsFramework/Private/EditorGizmos/GizmoElementRotateAxisSet.cpp
@@ -66,21 +66,6 @@ void UGizmoElementRotateAxisSet::DrawDebug(IToolsContextRenderAPI* RenderAPI, co
 	}
 }
 
-void UGizmoElementRotateAxisSet::SetGizmoViewContext(UGizmoViewContext* InGizmoViewContext)
-{
-	GizmoViewContext = InGizmoViewContext;
-
-	if (!bIsValid)
-	{
-		return;
-	}
-
-	ForEachAxisElement([&](UGizmoElementRotateAxis* InAxisElement)
-	{
-		InAxisElement->SetGizmoViewContext(GizmoViewContext);
-	});
-}
-
 void UGizmoElementRotateAxisSet::SetWidgetHost(IToolkitHost* const InWidgetHost)
 {
 	if (!bIsValid)
@@ -184,8 +169,6 @@ void UGizmoElementRotateAxisSet::Setup(
 			InParameters.Axis,
 			AxisStyle);
 
-		AxisElement->SetGizmoViewContext(GizmoViewContext);
-
 		return AxisElement;
 	};
 
diff --git a/Engine/Source/Editor/Experimental/EditorInteractiveToolsFramework/Private/EditorGizmos/GizmoElementRotateAxisSet.h b/Engine/Source/Editor/Experimental/EditorInteractiveToolsFramework/Private/EditorGizmos/GizmoElementRotateAxisSet.h
--- a/Engine/Source/Editor/Experimental/EditorInteractiveToolsFramework/Private/EditorGizmos/GizmoElementRotateAxisSet.h
+++ b/Engine/Source/Editor/Experimental/EditorInteractiveToolsFramework/Private/EditorGizmos/GizmoElementRotateAxisSet.h
@@ -60,9 +60,6 @@ public:
 	/** Renders various debug visualizations, if enabled. */
 	void DrawDebug(IToolsContextRenderAPI* RenderAPI, const FGizmoDebugSettings& InSettings, const uint32 InPartId = 0);
 
-	/** Sets the Gizmo view context, needed for screen space interactions */
-	void SetGizmoViewContext(UGizmoViewContext* InGizmoViewContext);
-
 	/** Sets the host to add/remove a widget to.
 	* Currently, this is a Toolkit host, but it could be abstracted to an ISlateWidgetHost (see IEditorViewportClientProxy). */
 	void SetWidgetHost(IToolkitHost* const InWidgetHost);
@@ -139,8 +136,4 @@ private:
 	/** Z axis rotation element. */
 	UPROPERTY()
 	TObjectPtr<UGizmoElementRotateAxis> AxisElementZ;
-
-    /** Gizmo view context, needed for screen space interactions. */
-    UPROPERTY(Setter)
-    TObjectPtr<UGizmoViewContext> GizmoViewContext;
 };
diff --git a/Engine/Source/Editor/Experimental/EditorInteractiveToolsFramework/Private/EditorGizmos/GizmoElementScaleAxis.cpp b/Engine/Source/Editor/Experimental/EditorInteractiveToolsFramework/Private/EditorGizmos/GizmoElementScaleAxis.cpp
--- a/Engine/Source/Editor/Experimental/EditorInteractiveToolsFramework/Private/EditorGizmos/GizmoElementScaleAxis.cpp
+++ b/Engine/Source/Editor/Experimental/EditorInteractiveToolsFramework/Private/EditorGizmos/GizmoElementScaleAxis.cpp
@@ -50,7 +50,7 @@ void UGizmoElementScaleAxis::Render(IToolsContextRenderAPI* RenderAPI, const FRe
 		return;
 	}
 
-	const float DPIScale = static_cast<float>(FMath::Max(UE_DOUBLE_KINDA_SMALL_NUMBER, RenderAPI->GetCameraState().DPIScale));
+	CurrentState.DPIScale = static_cast<float>(FMath::Max(UE_DOUBLE_KINDA_SMALL_NUMBER, RenderAPI->GetCameraState().DPIScale));
 
 	// To subtract from the regular line element - calculated here to account for PixelToWorldScale
 	float DeltaLineLength = 0.0f;
@@ -74,7 +74,7 @@ void UGizmoElementScaleAxis::Render(IToolsContextRenderAPI* RenderAPI, const FRe
 			DeltaLineLength = GetDeltaAxisLength(DeltaRenderState.LocalToWorldTransform);
 
 			// Adjust so that the delta line only appears for positive scale - for negative scale, it only appears when the handle is on the opposing side of the axis.
-			const float DeltaLineLengthChop = CurrentState.AxisSign < 0 ? GetTotalAxisLength() * DPIScale : 0.0f;
+			const float DeltaLineLengthChop = CurrentState.AxisSign < 0 ? GetTotalAxisLength() * CurrentState.DPIScale  : 0.0f;
 			const float DeltaLineLengthMultiplier = (CurrentState.DeltaSign == 1 || (CurrentState.DeltaSign == -1 && CurrentState.AxisSign == -1)) ? 1.0f : 0.0f;
 			const float AppliedDeltaLineLength = FMath::Abs((DeltaLineLength - DeltaLineLengthChop) * DeltaLineLengthMultiplier);
 
@@ -98,7 +98,7 @@ void UGizmoElementScaleAxis::Render(IToolsContextRenderAPI* RenderAPI, const FRe
 				// We only care about negative lengths, then invert to subtract from the body length
 				const float DeltaLineLengthToSubtract =
 					FMath::Abs(
-						FMath::Min(0.0f, (DeltaLineLength / DPIScale) * CurrentState.DeltaSign));
+						FMath::Min(0.0f, (DeltaLineLength / CurrentState.DPIScale) * CurrentState.DeltaSign));
 
 				ArrowBodyElement->SetHeight(FMath::Max(0.0f, GetAxisLength() - DeltaLineLengthToSubtract));	
 			}
@@ -111,7 +111,7 @@ void UGizmoElementScaleAxis::Render(IToolsContextRenderAPI* RenderAPI, const FRe
 
 		if (ArrowHeadElement)
 		{
-			ArrowHeadElement->SetCenter(GetHeadCenter(DPIScale));
+			ArrowHeadElement->SetCenter(GetHeadCenter(CurrentState.DPIScale));
 
 			ArrowHeadElement->Render(RenderAPI, CurrentRenderState);
 		}
@@ -153,11 +153,6 @@ void UGizmoElementScaleAxis::DrawDebug(IToolsContextRenderAPI* RenderAPI, const
 	}
 }
 
-void UGizmoElementScaleAxis::SetGizmoViewContext(UGizmoViewContext* InGizmoViewContext)
-{
-	GizmoViewContext = InGizmoViewContext;
-}
-
 void UGizmoElementScaleAxis::SetWidgetHost(IToolkitHost* const InWidgetHost)
 {
 	if (DeltaWidgetElement)
@@ -516,8 +511,7 @@ void UGizmoElementScaleAxis::UpdateDelta(const FDeltaParameters& InParameters)
 		}
 		else
 		{
-			const double DPIScale = GizmoViewContext ? GizmoViewContext->GetDPIScale() : 1.0;
-			DeltaLineElement->SetBase(NormalDirection * GetTotalAxisLength() * DPIScale);
+			DeltaLineElement->SetBase(NormalDirection * GetTotalAxisLength() * CurrentState.DPIScale );
 		}
 
 		DeltaLineElement->SetDirection(GetSignedDeltaDirection());
@@ -636,7 +630,13 @@ FVector UGizmoElementScaleAxis::GetSignedDeltaDirection() const
 
 bool UGizmoElementScaleAxis::UpdateDeltaRenderState(IToolsContextRenderAPI* RenderAPI, const FVector& InLocalOrigin, FRenderTraversalState& InOutRenderState)
 {
-	InOutRenderState.PixelToWorldScale = UE::GizmoRenderingUtil::CalculateLocalPixelToWorldScale(GizmoViewContext, StartState.Transform.GetLocation());
+	const FSceneView* SceneView = RenderAPI->GetSceneView();
+	check(SceneView);
+	const UE::GizmoRenderingUtil::FSceneViewWrapper SceneViewInterface(*SceneView);
+
+	InOutRenderState.PixelToWorldScale = UE::GizmoRenderingUtil::CalculateLocalPixelToWorldScale<double>(&SceneViewInterface, StartState.Transform.GetLocation());
+	const FVector4 ScreenSpaceStart = SceneView->WorldToScreen(StartState.Transform.GetLocation());
+	InOutRenderState.DepthAtPixelToWorldReferencePoint = ScreenSpaceStart.W;
 	InOutRenderState.LocalToWorldTransform.SetLocation(StartState.Transform.GetLocation());
 	InOutRenderState.LocalToWorldTransform.SetScale3D(FVector::OneVector);
 
diff --git a/Engine/Source/Editor/Experimental/EditorInteractiveToolsFramework/Private/EditorGizmos/GizmoElementScaleAxis.h b/Engine/Source/Editor/Experimental/EditorInteractiveToolsFramework/Private/EditorGizmos/GizmoElementScaleAxis.h
--- a/Engine/Source/Editor/Experimental/EditorInteractiveToolsFramework/Private/EditorGizmos/GizmoElementScaleAxis.h
+++ b/Engine/Source/Editor/Experimental/EditorInteractiveToolsFramework/Private/EditorGizmos/GizmoElementScaleAxis.h
@@ -92,9 +92,6 @@ public:
 	/** Renders various debug visualizations, if enabled. */
 	void DrawDebug(IToolsContextRenderAPI* RenderAPI, const FGizmoDebugSettings& InSettings);
 
-	/** Sets the Gizmo view context, needed for screen space interactions */
-	void SetGizmoViewContext(UGizmoViewContext* InGizmoViewContext);
-
 	/** Sets the host to add/remove a widget to.
 	* Currently, this is a Toolkit host, but it could be abstracted to an ISlateWidgetHost (see IEditorViewportClientProxy). */
 	void SetWidgetHost(IToolkitHost* const InWidgetHost);
@@ -195,10 +192,6 @@ private:
 private:
 	/** Flag indicating element validity - generally true after Setup is called. */
 	std::atomic_bool bIsValid = false;
-	
-	/** Gizmo view context, needed for screen space interactions */
-	UPROPERTY(Setter)
-	TObjectPtr<UGizmoViewContext> GizmoViewContext;
 
 	/** The axis this element represents (e.g. EAxisList::X). */
 	EAxisList::Type Axis = EAxisList::None;
@@ -236,6 +229,9 @@ private:
 		/** Whether the interaction is constrained to a single axis (as opposed to uniform/planar scale). */
 		bool bIsSingleAxis = false;
 
+		/** The current DPI Scale */
+		float DPIScale = 1.0f;
+
 		/** World-space point where the interaction ray intersects the drag plane. */
 		FVector PlaneIntersectionPoint = FVector::ZeroVector;
 
diff --git a/Engine/Source/Editor/Experimental/EditorInteractiveToolsFramework/Private/EditorGizmos/GizmoElementScaleGroup.cpp b/Engine/Source/Editor/Experimental/EditorInteractiveToolsFramework/Private/EditorGizmos/GizmoElementScaleGroup.cpp
--- a/Engine/Source/Editor/Experimental/EditorInteractiveToolsFramework/Private/EditorGizmos/GizmoElementScaleGroup.cpp
+++ b/Engine/Source/Editor/Experimental/EditorInteractiveToolsFramework/Private/EditorGizmos/GizmoElementScaleGroup.cpp
@@ -145,21 +145,6 @@ void UGizmoElementScaleGroup::DrawDebug(IToolsContextRenderAPI* RenderAPI, const
 	}
 }
 
-void UGizmoElementScaleGroup::SetGizmoViewContext(UGizmoViewContext* InGizmoViewContext)
-{
-	GizmoViewContext = InGizmoViewContext;
-
-	if (!bIsValid)
-	{
-		return;
-	}
-
-	ForEachAxisElement([&](UGizmoElementScaleAxis* InAxisElement)
-	{
-		InAxisElement->SetGizmoViewContext(GizmoViewContext);
-	});
-}
-
 void UGizmoElementScaleGroup::SetWidgetHost(IToolkitHost* const InWidgetHost)
 {
 	if (!bIsValid)
@@ -316,8 +301,6 @@ void UGizmoElementScaleGroup::Setup(
 				InParameters.Axis,
 				AxisStyle);
 
-			AxisElement->SetGizmoViewContext(GizmoViewContext);
-
 			return AxisElement;
 		};
 
diff --git a/Engine/Source/Editor/Experimental/EditorInteractiveToolsFramework/Private/EditorGizmos/GizmoElementScaleGroup.h b/Engine/Source/Editor/Experimental/EditorInteractiveToolsFramework/Private/EditorGizmos/GizmoElementScaleGroup.h
--- a/Engine/Source/Editor/Experimental/EditorInteractiveToolsFramework/Private/EditorGizmos/GizmoElementScaleGroup.h
+++ b/Engine/Source/Editor/Experimental/EditorInteractiveToolsFramework/Private/EditorGizmos/GizmoElementScaleGroup.h
@@ -133,9 +133,6 @@ public:
 	/** Renders various debug visualizations, if enabled. */
 	void DrawDebug(IToolsContextRenderAPI* RenderAPI, const FGizmoDebugSettings& InSettings, const uint32 InPartId = 0);
 
-	/** Sets the Gizmo view context, needed for screen space interactions */
-	void SetGizmoViewContext(UGizmoViewContext* InGizmoViewContext);
-
 	/** Sets the host to add/remove a widget to.
 	* Currently, this is a Toolkit host, but it could be abstracted to an ISlateWidgetHost (see IEditorViewportClientProxy). */
 	void SetWidgetHost(IToolkitHost* const InWidgetHost);
@@ -301,10 +298,6 @@ private:
 	UPROPERTY()
 	TObjectPtr<UGizmoElementSphere> UniformElementAlternate;
 
-    /** Gizmo view context, needed for screen space interactions */
-    UPROPERTY(Setter)
-    TObjectPtr<UGizmoViewContext> GizmoViewContext;
-
 	/** Callback that provides the current snapping configuration for scale interactions. */
 	TFunction<FToolContextSnappingConfiguration()> GetCurrentSnappingSettingsFunction = nullptr;
 
diff --git a/Engine/Source/Editor/Experimental/EditorInteractiveToolsFramework/Private/EditorGizmos/GizmoElementTranslateGroup.cpp b/Engine/Source/Editor/Experimental/EditorInteractiveToolsFramework/Private/EditorGizmos/GizmoElementTranslateGroup.cpp
--- a/Engine/Source/Editor/Experimental/EditorInteractiveToolsFramework/Private/EditorGizmos/GizmoElementTranslateGroup.cpp
+++ b/Engine/Source/Editor/Experimental/EditorInteractiveToolsFramework/Private/EditorGizmos/GizmoElementTranslateGroup.cpp
@@ -310,21 +310,6 @@ void UGizmoElementTranslateGroup::DrawDebug(IToolsContextRenderAPI* RenderAPI, c
 	}
 }
 
-void UGizmoElementTranslateGroup::SetGizmoViewContext(UGizmoViewContext* InGizmoViewContext)
-{
-	GizmoViewContext = InGizmoViewContext;
-
-	if (!bIsValid)
-	{
-		return;
-	}
-
-	ForEachAxisElement([&](UGizmoElementTranslateAxis* InAxisElement)
-	{
-		InAxisElement->SetGizmoViewContext(GizmoViewContext);
-	});
-}
-
 void UGizmoElementTranslateGroup::SetWidgetHost(IToolkitHost* const InWidgetHost)
 {
 	if (!bIsValid)
@@ -483,8 +468,6 @@ void UGizmoElementTranslateGroup::Setup(
 				InParameters.Axis,
 				AxisStyle);
 
-			AxisElement->SetGizmoViewContext(GizmoViewContext);
-
 			return AxisElement;
 		};
 
@@ -1076,7 +1059,22 @@ void UGizmoElementTranslateGroup::ForEachAxisElementByList(
 
 bool UGizmoElementTranslateGroup::UpdateDeltaRenderState(IToolsContextRenderAPI* RenderAPI, const FVector& InLocalOrigin, FRenderTraversalState& InOutRenderState)
 {
-	InOutRenderState.PixelToWorldScale = UE::GizmoRenderingUtil::CalculateLocalPixelToWorldScale<double>(GizmoViewContext, StartState.Translation);
+	// Use the SceneView from the RenderAPI which represents the current viewport.
+	// Not using GizmoViewContext, which might not have the right viewport info for the current viewport, in case of split views
+	const FSceneView* SceneView = RenderAPI->GetSceneView();
+	check(SceneView);
+	const UE::GizmoRenderingUtil::FSceneViewWrapper SceneViewInterface(*SceneView);
+
+	const FVector4 ScreenSpaceStart = SceneView->WorldToScreen(StartState.Translation);
+
+	// If StartState.Translation is behind the camera we need to use CurrentState information
+	// This might happen e.g. for Shift + Drag, camera moves while translating
+	const bool bIsBehindCamera = ScreenSpaceStart.W <= 0.0;
+	const FVector& ReferencePoint = bIsBehindCamera ? CurrentState.Translation : StartState.Translation;
+	const double ScreenSpaceDepth = bIsBehindCamera ? SceneView->WorldToScreen(CurrentState.Translation).W : ScreenSpaceStart.W;
+
+	InOutRenderState.PixelToWorldScale = UE::GizmoRenderingUtil::CalculateLocalPixelToWorldScale<double>(&SceneViewInterface, ReferencePoint);
+	InOutRenderState.DepthAtPixelToWorldReferencePoint = ScreenSpaceDepth;
 	InOutRenderState.LocalToWorldTransform.SetRotation(FQuat::Identity);
 	InOutRenderState.LocalToWorldTransform.SetTranslation(StartState.Translation);
 	InOutRenderState.LocalToWorldTransform.SetScale3D(FVector::OneVector);
diff --git a/Engine/Source/Editor/Experimental/EditorInteractiveToolsFramework/Private/EditorGizmos/GizmoElementTranslateGroup.h b/Engine/Source/Editor/Experimental/EditorInteractiveToolsFramework/Private/EditorGizmos/GizmoElementTranslateGroup.h
--- a/Engine/Source/Editor/Experimental/EditorInteractiveToolsFramework/Private/EditorGizmos/GizmoElementTranslateGroup.h
+++ b/Engine/Source/Editor/Experimental/EditorInteractiveToolsFramework/Private/EditorGizmos/GizmoElementTranslateGroup.h
@@ -130,9 +130,6 @@ public:
 	/** Renders various debug visualizations, if enabled. */
 	void DrawDebug(IToolsContextRenderAPI* RenderAPI, const FGizmoDebugSettings& InSettings, const uint32 InPartId = 0);
 
-	/** Sets the Gizmo view context, needed for screen space interactions */
-	void SetGizmoViewContext(UGizmoViewContext* InGizmoViewContext);
-
 	/** Sets the host to add/remove a widget to.
 	* Currently, this is a Toolkit host, but it could be abstracted to an ISlateWidgetHost (see IEditorViewportClientProxy). */
 	void SetWidgetHost(IToolkitHost* const InWidgetHost);
@@ -305,10 +302,6 @@ private:
 	UPROPERTY()
 	TObjectPtr<UGizmoElementCylinder> DeltaLineElement;
 
-    /** Gizmo view context, needed for screen space interactions */
-    UPROPERTY(Setter)
-    TObjectPtr<UGizmoViewContext> GizmoViewContext;
-
 	/** Callback that provides the current snapping configuration for interaction. */
 	TFunction<FToolContextSnappingConfiguration()> GetCurrentSnappingSettingsFunction = nullptr;
 
diff --git a/Engine/Source/Runtime/InteractiveToolsFramework/Private/BaseGizmos/GizmoElementCylinder.cpp b/Engine/Source/Runtime/InteractiveToolsFramework/Private/BaseGizmos/GizmoElementCylinder.cpp
--- a/Engine/Source/Runtime/InteractiveToolsFramework/Private/BaseGizmos/GizmoElementCylinder.cpp
+++ b/Engine/Source/Runtime/InteractiveToolsFramework/Private/BaseGizmos/GizmoElementCylinder.cpp
@@ -90,7 +90,8 @@ void UGizmoElementCylinder::Render(IToolsContextRenderAPI* RenderAPI, const FRen
 			// Get screen-space endpoints (used for visibility and for reparameterization, both world and screen-space)
             FVector ScreenStartPoint, ScreenEndPoint;
             double LineFractionInView = 1.0;
-           GizmoMath::GetScreenSpaceLineStartEnd(RenderAPI->GetSceneView(), WorldOrigin, WorldDirection, WorldHeight, ScreenStartPoint, ScreenEndPoint, LineFractionInView);
+            double WorldStartFraction = 0.0;
+			GizmoMath::GetScreenSpaceLineStartEnd(RenderAPI->GetSceneView(), WorldOrigin, WorldDirection, WorldHeight, ScreenStartPoint, ScreenEndPoint, LineFractionInView, WorldStartFraction);
 
 			const float HeightInView = Height * LineFractionInView;
 
@@ -110,6 +111,7 @@ void UGizmoElementCylinder::Render(IToolsContextRenderAPI* RenderAPI, const FRen
 				DashLineSettings.ScreenEndPoint = ScreenEndPoint;
 				DashLineSettings.ScaledLength = Height;
 				DashLineSettings.LineFractionInView = LineFractionInView;
+				DashLineSettings.WorldStartFraction = WorldStartFraction;
 
 				if (bScreenSpace)
 				{
@@ -125,9 +127,16 @@ void UGizmoElementCylinder::Render(IToolsContextRenderAPI* RenderAPI, const FRen
 						return (ClampedParam * StartDepth) / (((1.0 - ClampedParam) * EndDepth) + (ClampedParam * StartDepth));
 					};
 
-					DashLineSettings.ParamToRadiusScaleFunction = [StartDepth = ScreenStartPoint.Z, EndDepth = ScreenEndPoint.Z](const double& InParam)
+					double DepthGapCorrection = 1.0;
+					const double ReferencePointDepth = RenderState.DepthAtPixelToWorldReferencePoint;
+					if (ReferencePointDepth > 0)
 					{
-						return GizmoMath::GetPixelToWorldScaleAtScreenFraction(StartDepth, EndDepth, InParam);
+						// When PixelToWorldScale is not sampled at StartState, StartDepth and ReferencePointDepth are not the same (see UGizmoElementTranslateGroup::UpdateDeltaRenderState)
+						DepthGapCorrection = ScreenStartPoint.Z / ReferencePointDepth;
+					}
+					DashLineSettings.ParamToRadiusScaleFunction = [DepthGapCorrection, StartDepth = ScreenStartPoint.Z, EndDepth = ScreenEndPoint.Z](const double& InParam)
+					{
+						return DepthGapCorrection * GizmoMath::GetPixelToWorldScaleAtScreenFraction(StartDepth, EndDepth, InParam);
 					};
 
 					GizmoMath::BuildDashedLine(RenderAPI->GetPrimitiveDrawInterface(), DashLineSettings);
diff --git a/Engine/Source/Runtime/InteractiveToolsFramework/Private/BaseGizmos/GizmoMath.cpp b/Engine/Source/Runtime/InteractiveToolsFramework/Private/BaseGizmos/GizmoMath.cpp
--- a/Engine/Source/Runtime/InteractiveToolsFramework/Private/BaseGizmos/GizmoMath.cpp
+++ b/Engine/Source/Runtime/InteractiveToolsFramework/Private/BaseGizmos/GizmoMath.cpp
@@ -771,42 +771,44 @@ namespace GizmoMath
 		const double& InLineLength,
 		FVector& OutScreenStartPoint,
 		FVector& OutScreenEndPoint,
-		double& OutClampedWorldFraction
+		double& OutClampedWorldFraction,
+		double& OutClampedWorldStartFraction
 	)
 	{
-		const FIntRect ViewRect = InSceneView->UnscaledViewRect;
+		// Facilitator method to quickly return the result of InSceneView->ScreenToPixel
+		auto ScreenToPixel = [&](const FVector4& InScreenPoint)
+		{
+			FVector2D OutPoint = FVector2D::ZeroVector;
+			InSceneView->ScreenToPixel(InScreenPoint, OutPoint);
+			return OutPoint;
+		};
 
-		FVector LineStart = InLineOrigin;
-		FVector LineEnd = InLineOrigin + InLineDirection * InLineLength;
+		const FIntRect& ViewRect = InSceneView->UnscaledViewRect;
 
-		FVector4 ScreenStartPoint = InSceneView->WorldToScreen(LineStart);
-		FVector4 ScreenEndPoint = InSceneView->WorldToScreen(LineEnd);
+		const FVector& LineStart = InLineOrigin;
+		const FVector LineEnd = InLineOrigin + InLineDirection * InLineLength;
 
-		FVector2D UnclampedPixelStartPoint = FVector2D::ZeroVector;
-		InSceneView->ScreenToPixel(ScreenStartPoint, UnclampedPixelStartPoint);
+		const FVector4 ScreenStartPoint = InSceneView->WorldToScreen(LineStart);
+		const FVector4 ScreenEndPoint = InSceneView->WorldToScreen(LineEnd);
 
-		FVector2D UnclampedPixelEndPoint = FVector2D::ZeroVector;
-		InSceneView->ScreenToPixel(ScreenEndPoint, UnclampedPixelEndPoint);
+		// Pixel coordinates of line Start and End. These could be outside the viewport
+		const FVector2D PixelStartPoint = ScreenToPixel(ScreenStartPoint);
+		const FVector2D PixelEndPoint = ScreenToPixel(ScreenEndPoint);
 
-		auto ClampToViewRect = [&](const FVector2D& InPoint)
+		// Checks if the provided point is within the viewport or not
+		auto IsOutsideViewRect = [&](const FVector2D& InPoint)
 		{
-			// @see: FSceneView::CursorToScreen
-			return FVector2D::Clamp(
-				InPoint, FVector2D(ViewRect.Min.X, ViewRect.Min.Y), FVector2D(ViewRect.Max.X, ViewRect.Max.Y)
-			);
+			return InPoint.X < ViewRect.Min.X || InPoint.X > ViewRect.Max.X || InPoint.Y < ViewRect.Min.Y
+				|| InPoint.Y > ViewRect.Max.Y;
 		};
 
-		const FVector2D ClampedPixelStartPoint = ClampToViewRect(UnclampedPixelStartPoint);
-		const FVector2D ClampedPixelEndPoint = ClampToViewRect(UnclampedPixelEndPoint);
-
-		const double UnclampedLineLength = FVector2D::Distance(UnclampedPixelStartPoint, UnclampedPixelEndPoint);
-		const double ClampedLineLength = FVector2D::Distance(ClampedPixelStartPoint, ClampedPixelEndPoint);
+		const bool bNeedsClamping = IsOutsideViewRect(PixelStartPoint) || IsOutsideViewRect(PixelEndPoint);
 
-		const bool bIsClamped = !FMath::IsNearlyEqual(UnclampedLineLength, ClampedLineLength);
-
-		// If clamped, we need to adjust the Z values to be at the clamped position along the line
-		if (bIsClamped)
+		// If clamping is needed, we need to find their position, and re-compute screen and pixel coordinates
+		if (bNeedsClamping)
 		{
+			FVector ClampedLineStart = LineStart;
+			FVector ClampedLineEnd = LineEnd;
 			// @note: the frustum clipping code below could be cached, with start and end intersection world points
 			// which are then used if the corresponding point is out of bounds
 			// @see: LevelEditorViewport.cpp TrimLineToFrustum
@@ -817,23 +819,34 @@ namespace GizmoMath
 				{
 					if ((FrustumPlane.Dot(IntersectionPoint - LineEnd) > UE_DOUBLE_SMALL_NUMBER))
 					{
-						LineStart = IntersectionPoint;
+						ClampedLineStart = IntersectionPoint;
 					}
 					else
 					{
-						LineEnd = IntersectionPoint;
+						ClampedLineEnd = IntersectionPoint;
 					}
 				}
 			}
 
-			OutClampedWorldFraction = FVector::Distance(LineStart, LineEnd) / InLineLength;
+			OutClampedWorldFraction = FVector::Distance(ClampedLineStart, ClampedLineEnd) / InLineLength;
+			OutClampedWorldStartFraction = FVector::Distance(LineStart, ClampedLineStart) / InLineLength;
 
-			ScreenStartPoint = InSceneView->WorldToScreen(LineStart);
-			ScreenEndPoint = InSceneView->WorldToScreen(LineEnd);
-		}
+			// We have new Clamped Start and End of line: screen and pixel space points have to be re-computed
+			const FVector4 ClampedScreenStartPoint = InSceneView->WorldToScreen(ClampedLineStart);
+			const FVector4 ClampedScreenEndPoint = InSceneView->WorldToScreen(ClampedLineEnd);
+
+			const FVector2D ClampedPixelStartPoint = ScreenToPixel(ClampedScreenStartPoint);
+			const FVector2D ClampedPixelEndPoint = ScreenToPixel(ClampedScreenEndPoint);
 
-		OutScreenStartPoint = FVector(ClampedPixelStartPoint.X, ClampedPixelStartPoint.Y, ScreenStartPoint.W);
-		OutScreenEndPoint = FVector(ClampedPixelEndPoint.X, ClampedPixelEndPoint.Y, ScreenEndPoint.W);
+			OutScreenStartPoint = FVector(ClampedPixelStartPoint.X, ClampedPixelStartPoint.Y, ClampedScreenStartPoint.W);
+			OutScreenEndPoint = FVector(ClampedPixelEndPoint.X, ClampedPixelEndPoint.Y, ClampedScreenEndPoint.W);
+		}
+		else
+		{
+			OutScreenStartPoint = FVector(PixelStartPoint.X, PixelStartPoint.Y, ScreenStartPoint.W);
+			OutScreenEndPoint = FVector(PixelEndPoint.X, PixelEndPoint.Y, ScreenEndPoint.W);
+			OutClampedWorldStartFraction = 0.0;
+		}
 
 		// Clip to near plane if needed
 		OutScreenStartPoint.Z = FMath::Max(OutScreenStartPoint.Z, UE_DOUBLE_SMALL_NUMBER);
@@ -908,8 +921,14 @@ namespace GizmoMath
 			const double WorldDashMid = 0.5 * (InWorldStart + InWorldEnd);
 			const double HalfDashLength = 0.5 * (InWorldEnd - InWorldStart) * HeightInView;
 
+			// WorldVisibleStart "offsets" the base location according to the cut portion of the dash
+			// This happens when the line origin sits outside the viewport (e.g. Gizmo Shift + Drag translate which moves the camera during translation)
+			const float WorldVisibleStart = InDashSettings.ScaledLength * InDashSettings.WorldStartFraction;
+
+			const FVector Base = (HeightInView * WorldDashMid + WorldVisibleStart - HalfLength) * FVector::UpVector;
+
 			BuildTaperedCylinderVerts(
-				(-FVector::UpVector * HalfLength) + FVector::UpVector * (HeightInView * WorldDashMid),
+				Base,
 				FVector::XAxisVector,
 				FVector::YAxisVector,
 				FVector::ZAxisVector,
@@ -920,8 +939,6 @@ namespace GizmoMath
 				OutVertices,
 				OutTriangleIndices
 			);
-
-			
 		};
 
 		auto AppendDash = [&](const double& InStart, const double& InEnd)
diff --git a/Engine/Source/Runtime/InteractiveToolsFramework/Public/BaseGizmos/GizmoElementBase.h b/Engine/Source/Runtime/InteractiveToolsFramework/Public/BaseGizmos/GizmoElementBase.h
--- a/Engine/Source/Runtime/InteractiveToolsFramework/Public/BaseGizmos/GizmoElementBase.h
+++ b/Engine/Source/Runtime/InteractiveToolsFramework/Public/BaseGizmos/GizmoElementBase.h
@@ -61,6 +61,9 @@ public:
 		// Pixel to world scale. Also accounts for DPI scale, @see IToolsContextRenderAPI::ViewCameraState.
 		double PixelToWorldScale = 1.0;
 
+		// The screen space depth of the reference point used to compute PixelToWorldScale. Defaults to -1.0 (unset). Should not be used if negative.
+		double DepthAtPixelToWorldReferencePoint = -1.0;
+
 		// Interact state, if not equal to none, overrides the element's interact state 
 		EGizmoElementInteractionState InteractionState = EGizmoElementInteractionState::None;
 
diff --git a/Engine/Source/Runtime/InteractiveToolsFramework/Public/BaseGizmos/GizmoMath.h b/Engine/Source/Runtime/InteractiveToolsFramework/Public/BaseGizmos/GizmoMath.h
--- a/Engine/Source/Runtime/InteractiveToolsFramework/Public/BaseGizmos/GizmoMath.h
+++ b/Engine/Source/Runtime/InteractiveToolsFramework/Public/BaseGizmos/GizmoMath.h
@@ -203,6 +203,9 @@ namespace GizmoMath
 		/** Reflects the fraction of the line (in world space) that's visible - e.g. 0.5 means half of the line is visible */
 		double LineFractionInView = 1.0;
 
+		/** World-space fraction offset - non-zero only when the line origin is clamped */
+		double WorldStartFraction = 0.0;
+
 		/** Radius of the line */
 		float Radius = 1.0f;
 
@@ -221,6 +224,8 @@ namespace GizmoMath
 	 * are in pixel space. OutClampedWorldFraction is used when one or more of the end points is beyond the view bounds,
 	 * and reflects the fraction of the line (in world space) that's visible, ie. if only 50% of the line is within the
 	 * viewport, this will return 0.5.
+	 * OutClampedWorldStartFraction reflects the world-space offset from the line origin to the visible segment start -
+	 * non-zero when the origin end of the line is clipped (e.g. camera-follow mode where the origin leaves the viewport).
 	 */
 	INTERACTIVETOOLSFRAMEWORK_API void GetScreenSpaceLineStartEnd(
 		const FSceneView* InSceneView,
@@ -229,7 +234,8 @@ namespace GizmoMath
 		const double& InLineLength,
 		FVector& OutScreenStartPoint,
 		FVector& OutScreenEndPoint,
-		double& OutClampedWorldFraction
+		double& OutClampedWorldFraction,
+		double& OutClampedWorldStartFraction
 	);
 
 	INTERACTIVETOOLSFRAMEWORK_API double GetPixelToWorldScaleAtScreenFraction(
